SessionServices: Documentation cleanup pass

diff --git a/Engine/Source/Developer/SessionMessages/Classes/SessionServiceMessages.h b/Engine/Source/Developer/SessionMessages/Classes/SessionServiceMessages.h
--- a/Engine/Source/Developer/SessionMessages/Classes/SessionServiceMessages.h
+++ b/Engine/Source/Developer/SessionMessages/Classes/SessionServiceMessages.h
@@ -40,63 +40,43 @@ struct FSessionServicePong
 {
 	GENERATED_USTRUCT_BODY()
 
-	/**
-	 * Holds the application's build date.
-	 */
+	/** Holds the application's build date. */
 	UPROPERTY()
 	FString BuildDate;
 
-	/**
-	 * Holds the name of the device that the application is running on.
-	 */
+	/** Holds the name of the device that the application is running on. */
 	UPROPERTY()
 	FString DeviceName;
 
-	/**
-	 * Holds the application's instance identifier.
-	 */
+	/** Holds the application's instance identifier. */
 	UPROPERTY()
 	FGuid InstanceId;
 
-	/**
-	 * Holds the application's instance name.
-	 */
+	/** Holds the application's instance name. */
 	UPROPERTY()
 	FString InstanceName;
 
-	/**
-	 * Holds a flag indicating whether the application is running on a console.
-	 */
+	/** Holds a flag indicating whether the application is running on a console. */
 	UPROPERTY()
 	bool IsConsoleBuild;
 
-	/**
-	 * Holds the name of the platform that the application is running on.
-	 */
+	/** Holds the name of the platform that the application is running on. */
 	UPROPERTY()
 	FString PlatformName;
 
-	/**
-	 * Holds the identifier of the session that the application belongs to.
-	 */
+	/** Holds the identifier of the session that the application belongs to. */
 	UPROPERTY()
 	FGuid SessionId;
 
-	/**
-	 * Holds the user defined name of the session.
-	 */
+	/** Holds the user defined name of the session. */
 	UPROPERTY()
 	FString SessionName;
 
-	/**
-	 * Holds the name of the user that started the session.
-	 */
+	/** Holds the name of the user that started the session. */
 	UPROPERTY()
 	FString SessionOwner;
 
-	/**
-	 * Holds a flag indicating whether the application is the only one in that session.
-	 */
+	/** Holds a flag indicating whether the application is the only one in that session. */
 	UPROPERTY()
 	bool Standalone;
 };
@@ -122,36 +102,27 @@ struct FSessionServiceLog
 {
 	GENERATED_USTRUCT_BODY()
 
-	/**
-	 * Holds the log message category.
-	 */
+	/** Holds the log message category. */
 	UPROPERTY()
 	FName Category;
 
-	/**
-	 * Holds the log message data.
-	 */
+	/** Holds the log message data. */
 	UPROPERTY()
 	FString Data;
 
-	/**
-	 * Holds the application instance indentifier.
-	 */
+	/** Holds the application instance identifier. */
 	UPROPERTY()
 	FGuid InstanceId;
 
-	/**
-	 * Holds the time in seconds since the application was started.
-	 */
+	/** Holds the time in seconds since the application was started. */
 	UPROPERTY()
 	double TimeSeconds;
 
-	/**
-	 * Holds the log message's verbosity level.
-	 */
+	/** Holds the log message's verbosity level. */
 	UPROPERTY()
 	uint8 Verbosity;
 
+public:
 
 	/**
 	 * Default constructor.
diff --git a/Engine/Source/Developer/SessionServices/Private/SessionInfo.h b/Engine/Source/Developer/SessionServices/Private/SessionInfo.h
--- a/Engine/Source/Developer/SessionServices/Private/SessionInfo.h
+++ b/Engine/Source/Developer/SessionServices/Private/SessionInfo.h
@@ -19,8 +19,8 @@ public:
 	/**
 	 * Creates and initializes a new instance.
 	 *
-	 * @param InSessionId - The session's identifier.
-	 * @param InMessageBus - The message bus to use.
+	 * @param InSessionId The session's identifier.
+	 * @param InMessageBus The message bus to use.
 	 */
 	FSessionInfo( const FGuid& InSessionId, const IMessageBusRef& InMessageBus )
 		: MessageBusPtr(InMessageBus)
@@ -32,8 +32,8 @@ public:
 	/**
 	 * Updates this session info with the data in the specified message.
 	 *
-	 * @param Message - The message containing engine information.
-	 * @param Context - The message context.
+	 * @param Message The message containing engine information.
+	 * @param Context The message context.
 	 */
 	void UpdateFromMessage( const FEngineServicePong& Message, const IMessageContextRef& Context )
 	{
@@ -63,8 +63,8 @@ public:
 	/**
 	 * Updates this session info with the data in the specified message.
 	 *
-	 * @param Message - The message containing session information.
-	 * @param Context - The message context.
+	 * @param Message The message containing session information.
+	 * @param Context The message context.
 	 */
 	void UpdateFromMessage( const FSessionServicePong& Message, const IMessageContextRef& Context )
 	{
@@ -181,7 +181,7 @@ private:
 private:
 
 	// Holds the list of engine instances that belong to this session.
-	TMap<FMessageAddress, TSharedPtr<FSessionInstanceInfo> > Instances;
+	TMap<FMessageAddress, TSharedPtr<FSessionInstanceInfo>> Instances;
 
 	// Holds the time at which the last pong was received.
 	FDateTime LastUpdateTime;
diff --git a/Engine/Source/Developer/SessionServices/Private/SessionInstanceInfo.h b/Engine/Source/Developer/SessionServices/Private/SessionInstanceInfo.h
--- a/Engine/Source/Developer/SessionServices/Private/SessionInstanceInfo.h
+++ b/Engine/Source/Developer/SessionServices/Private/SessionInstanceInfo.h
@@ -24,9 +24,9 @@ public:
 	/**
 	 * Creates and initializes a new instance.
 	 *
-	 * @param InInstanceId - The instance's identifier.
-	 * @param InOwner - The session that owns this instance.
-	 * @param InMessageBus - The message bus to use.
+	 * @param InInstanceId The instance's identifier.
+	 * @param InOwner The session that owns this instance.
+	 * @param InMessageBus The message bus to use.
 	 */
 	FSessionInstanceInfo( const FGuid& InInstanceId, const ISessionInfoRef& InOwner, const IMessageBusRef& InMessageBus )
 		: EngineVersion(0)
@@ -42,8 +42,8 @@ public:
 	/**
 	 * Updates this instance info with the data in the specified message.
 	 *
-	 * @param Message - The message containing engine information.
-	 * @param Context - The message context.
+	 * @param Message The message containing engine information.
+	 * @param Context The message context.
 	 */
 	void UpdateFromMessage( const FEngineServicePong& Message, const IMessageContextRef& Context )
 	{
@@ -63,8 +63,8 @@ public:
 	/**
 	 * Updates this instance info with the data in the specified message.
 	 *
-	 * @param Message - The message containing instance information.
-	 * @param Context - The message context.
+	 * @param Message The message containing instance information.
+	 * @param Context The message context.
 	 */
 	void UpdateFromMessage( const FSessionServicePong& Message, const IMessageContextRef& Context )
 	{
diff --git a/Engine/Source/Developer/SessionServices/Private/SessionManager.h b/Engine/Source/Developer/SessionServices/Private/SessionManager.h
--- a/Engine/Source/Developer/SessionServices/Private/SessionManager.h
+++ b/Engine/Source/Developer/SessionServices/Private/SessionManager.h
@@ -19,7 +19,7 @@ public:
 	/**
 	 * Creates and initializes a new instance.
 	 *
-	 * @param InMessageBus - The message bus to use.
+	 * @param InMessageBus The message bus to use.
 	 */
 	FSessionManager( const IMessageBusRef& InMessageBus );
 	
@@ -91,7 +91,7 @@ protected:
 	/**
 	 * Finds and removes sessions that haven't been updated in a while.
 	 *
-	 * @param Now - The current time.
+	 * @param Now The current time.
 	 */
 	void FindExpiredSessions( const FDateTime& Now );
 
@@ -114,7 +114,7 @@ protected:
 
 private:
 
-	// Handles FSessionServicePong messages.
+	// Callback for handling FSessionServicePong messages.
 	void HandleEnginePongMessage( const FEngineServicePong& Message, const IMessageContextRef& Context );
 
 	// Callback for newly discovered instances.
@@ -123,7 +123,7 @@ private:
 	// Callback received log entries.
 	void HandleLogReceived( const ISessionInfoRef& Session, const ISessionInstanceInfoRef& Instance, const FSessionLogMessageRef& Message );
 
-	// Handles FSessionServicePong messages.
+	// Callback for handling FSessionServicePong messages.
 	void HandleSessionPongMessage( const FSessionServicePong& Message, const IMessageContextRef& Context );
 
 	// Callback for ticks from the ticker.
diff --git a/Engine/Source/Developer/SessionServices/Private/SessionService.h b/Engine/Source/Developer/SessionServices/Private/SessionService.h
--- a/Engine/Source/Developer/SessionServices/Private/SessionService.h
+++ b/Engine/Source/Developer/SessionServices/Private/SessionService.h
@@ -19,7 +19,7 @@ public:
 	/**
 	 * Creates and initializes a new instance.
 	 *
-	 * @param InMessageBus - The message bus to use.
+	 * @param InMessageBus The message bus to use.
 	 */
 	FSessionService( const IMessageBusRef& InMessageBus );
 
@@ -28,7 +28,6 @@ public:
 	 */
 	~FSessionService( );
 
-
 public:
 
 	// Begin FOutputDevice interface
@@ -40,7 +39,6 @@ public:
 
 	// End FOutputDevice interface
 
-
 public:
 
 	// Begin ISessionService interface
@@ -56,34 +54,32 @@ public:
 
 	// End ISessionService interface
 
-
 protected:
 
 	/**
 	 * Sends a log message to subscribed recipients.
 	 *
-	 * @param Data - The log message data.
-	 * @param Verbosity - The verbosity type.
-	 * @param Category - The log category.
+	 * @param Data The log message data.
+	 * @param Verbosity The verbosity type.
+	 * @param Category The log category.
 	 */
 	void SendLog( const TCHAR* Data, ELogVerbosity::Type Verbosity = ELogVerbosity::Log, const class FName& Category = "Log" );
 
 	/**
 	 * Sends a notification to the specified recipient.
 	 *
-	 * @param NotificationText - The notification text.
-	 * @param Recipient - The recipient's message address.
+	 * @param NotificationText The notification text.
+	 * @param Recipient The recipient's message address.
 	 */
 	void SendNotification( const TCHAR* NotificationText, const FMessageAddress& Recipient );
 
 	/**
 	 * Publishes a ping response.
 	 *
-	 * @param Context - The context of the received Ping message.
+	 * @param Context The context of the received Ping message.
 	 */
 	void SendPong( const IMessageContextRef& Context );
 
-
 private:
 
 	// Handles message bus shutdowns.
@@ -98,7 +94,6 @@ private:
 	// Handles FSessionServicePing messages.
 	void HandleSessionPingMessage( const FSessionServicePing& Message, const IMessageContextRef& Context );
 
-
 private:
 
 	// Holds the list of log subscribers.
@@ -112,4 +107,4 @@ private:
 
 	// Holds the message endpoint.
 	FMessageEndpointPtr MessageEndpoint;
-};
\ No newline at end of file
+};
diff --git a/Engine/Source/Developer/SessionServices/Public/Interfaces/ISessionInfo.h b/Engine/Source/Developer/SessionServices/Public/Interfaces/ISessionInfo.h
--- a/Engine/Source/Developer/SessionServices/Public/Interfaces/ISessionInfo.h
+++ b/Engine/Source/Developer/SessionServices/Public/Interfaces/ISessionInfo.h
@@ -7,14 +7,10 @@
 #pragma once
 
 
-/**
- * Type definition for shared pointers to instances of ISessionInfo.
- */
+/** Type definition for shared pointers to instances of ISessionInfo. */
 typedef TSharedPtr<class ISessionInfo> ISessionInfoPtr;
 
-/**
- * Type definition for shared references to instances of ISessionInfo.
- */
+/** Type definition for shared references to instances of ISessionInfo. */
 typedef TSharedRef<class ISessionInfo> ISessionInfoRef;
 
 
@@ -46,7 +42,7 @@ public:
 	/**
 	 * Gets a read-only collection of all instances that belong to this session.
 	 *
-	 * @param OutInstances - Will hold the collection of instances.
+	 * @param OutInstances Will hold the collection of instances.
 	 */
 	virtual void GetInstances( TArray<ISessionInstanceInfoPtr>& OutInstances ) const = 0;
 
@@ -90,7 +86,7 @@ public:
 	 *
 	 * A session is standalone if has not been created from the Launcher.
 	 *
-	 * @return True if this is a standalone session, false otherwise.
+	 * @return true if this is a standalone session, false otherwise.
 	 */
 	virtual const bool IsStandalone( ) const = 0;
 
diff --git a/Engine/Source/Developer/SessionServices/Public/Interfaces/ISessionInstanceInfo.h b/Engine/Source/Developer/SessionServices/Public/Interfaces/ISessionInstanceInfo.h
--- a/Engine/Source/Developer/SessionServices/Public/Interfaces/ISessionInstanceInfo.h
+++ b/Engine/Source/Developer/SessionServices/Public/Interfaces/ISessionInstanceInfo.h
@@ -7,14 +7,10 @@
 #pragma once
 
 
-/**
- * Type definition for shared pointers to instances of IGameInstanceInfo.
- */
+/** Type definition for shared pointers to instances of IGameInstanceInfo. */
 typedef TSharedPtr<class ISessionInstanceInfo> ISessionInstanceInfoPtr;
 
-/**
- * Type definition for shared references to instances of IGameInstanceInfo.
- */
+/** Type definition for shared references to instances of IGameInstanceInfo. */
 typedef TSharedRef<class ISessionInstanceInfo> ISessionInstanceInfoRef;
 
 
@@ -37,7 +33,7 @@ public:
 	/**
 	 * Executes a console command on this engine instance.
 	 *
-	 * @param CommandString - The command to execute.
+	 * @param CommandString The command to execute.
 	 */
 	virtual void ExecuteCommand( const FString& CommandString ) = 0;
 
diff --git a/Engine/Source/Developer/SessionServices/Public/Interfaces/ISessionManager.h b/Engine/Source/Developer/SessionServices/Public/Interfaces/ISessionManager.h
--- a/Engine/Source/Developer/SessionServices/Public/Interfaces/ISessionManager.h
+++ b/Engine/Source/Developer/SessionServices/Public/Interfaces/ISessionManager.h
@@ -7,14 +7,10 @@
 #pragma once
 
 
-/**
- * Type definition for shared pointers to instances of ISessionManager.
- */
+/** Type definition for shared pointers to instances of ISessionManager. */
 typedef TSharedPtr<class ISessionManager> ISessionManagerPtr;
 
-/**
- * Type definition for shared references to instances of ISessionManager.
- */
+/** Type definition for shared references to instances of ISessionManager. */
 typedef TSharedRef<class ISessionManager> ISessionManagerRef;
 
 
@@ -51,28 +47,28 @@ public:
 	/**
 	 * Adds an owner whose sessions we are interested in
 	 *
-	 * @param InOwner - Session owner we want to view sessions from
+	 * @param InOwner Session owner we want to view sessions from
 	 */
 	virtual void AddOwner( const FString& InOwner ) = 0;
 
 	/**
 	  * Gets the collection of currently selected engine instances.
 	  *
-	  * @param OutInstances - Will hold the selected instances.
+	  * @param OutInstances Will hold the selected instances.
 	  */
 	virtual void GetSelectedInstances( TArray<ISessionInstanceInfoPtr>& OutInstances) const = 0;
 
 	/**
 	 * Get the selected session - as chosen in the session browser
 	 *
-	 * @Return	The session ID selected
+	 * @return The session ID selected
 	 */
 	virtual const ISessionInfoPtr& GetSelectedSession( ) const = 0;
 
 	/**
 	 * Gets the list of all discovered sessions.
 	 *
-	 * @param OutSessions - Will hold the collection of sessions.
+	 * @param OutSessions Will hold the collection of sessions.
 	 */
 	virtual void GetSessions( TArray<ISessionInfoPtr>& OutSessions ) const = 0;
 
@@ -86,15 +82,14 @@ public:
 	/**
 	 * Removes an owner whose sessions we are no longer interested in
 	 *
-	 * @param InOwner - Session owner we want to remove
+	 * @param InOwner Session owner we want to remove
 	 */
 	virtual void RemoveOwner( const FString& InOwner ) = 0;
 
 	/**
 	 * Selects the specified session.
 	 *
-	 * @param Session - The session to the select (can be NULL to select none).
-	 *
+	 * @param Session The session to the select (can be NULL to select none).
 	 * @return true if the session was selected, false otherwise.
 	 */
 	virtual bool SelectSession( const ISessionInfoPtr& Session ) = 0;
@@ -102,9 +97,8 @@ public:
 	/**
 	 * Marks the specified item as selected or unselected.
 	 *
-	 * @param Instance - The instance to mark.
-	 * @param Selected - Whether the instance should be selected (true) or unselected (false).
-	 *
+	 * @param Instance The instance to mark.
+	 * @param Selected Whether the instance should be selected (true) or unselected (false).
 	 * @return true if the instance was selected, false otherwise.
 	 */
 	virtual bool SetInstanceSelected( const ISessionInstanceInfoPtr& Instance, bool Selected ) = 0;
diff --git a/Engine/Source/Developer/SessionServices/Public/Interfaces/ISessionService.h b/Engine/Source/Developer/SessionServices/Public/Interfaces/ISessionService.h
--- a/Engine/Source/Developer/SessionServices/Public/Interfaces/ISessionService.h
+++ b/Engine/Source/Developer/SessionServices/Public/Interfaces/ISessionService.h
@@ -7,14 +7,10 @@
 #pragma once
 
 
-/**
- * Type definition for shared pointers to instances of ISessionService.
- */
+/** Type definition for shared pointers to instances of ISessionService. */
 typedef TSharedPtr<class ISessionService> ISessionServicePtr;
 
-/**
- * Type definition for shared references to instances of ISessionService.
- */
+/** Type definition for shared references to instances of ISessionService. */
 typedef TSharedRef<class ISessionService> ISessionServiceRef;
 
 
@@ -36,14 +32,12 @@ public:
 	 * Starts the service.
 	 *
 	 * @return true if the service was started, false otherwise.
-	 *
 	 * @see Stop
 	 */
 	virtual bool Start( ) = 0;
 
 	/**
 	 * Stops the service.
-	 *
 	 * @see Start
 	 */
 	virtual void Stop( ) = 0;
diff --git a/Engine/Source/Developer/SessionServices/Public/Interfaces/SessionLogMessage.h b/Engine/Source/Developer/SessionServices/Public/Interfaces/SessionLogMessage.h
--- a/Engine/Source/Developer/SessionServices/Public/Interfaces/SessionLogMessage.h
+++ b/Engine/Source/Developer/SessionServices/Public/Interfaces/SessionLogMessage.h
@@ -7,14 +7,10 @@
 #pragma once
 
 
-/**
- * Type definition for shared pointers to instances of FLauncherLogMessage.
- */
+/** Type definition for shared pointers to instances of FLauncherLogMessage. */
 typedef TSharedPtr<struct FSessionLogMessage> FSessionLogMessagePtr;
 
-/**
- * Type definition for shared references to instances of FLauncherLogMessage.
- */
+/** Type definition for shared references to instances of FLauncherLogMessage. */
 typedef TSharedRef<struct FSessionLogMessage> FSessionLogMessageRef;
 
 
@@ -23,51 +19,38 @@ typedef TSharedRef<struct FSessionLogMessage> FSessionLogMessageRef;
  */
 struct FSessionLogMessage
 {
-	/**
-	 * Holds the log category.
-	 */
+	/** Holds the log category. */
 	FName Category;
 
-	/**
-	 * Holds the identifier of the engine instance that generated this log message.
-	 */
+	/** Holds the identifier of the engine instance that generated this log message. */
 	FGuid InstanceId;
 
-	/**
-	 * Holds the name of the engine instance that generated this log message.
-	 */
+	/** Holds the name of the engine instance that generated this log message. */
 	FString InstanceName;
 
-	/**
-	 * Holds the message text.
-	 */
+	/** Holds the message text. */
 	FString Text;
 
-	/**
-	 * Holds the time at which the message was generated.
-	 */
+	/** Holds the time at which the message was generated. */
 	FDateTime Time;
 
-	/**
-	 * Holds the number of seconds from the start of the instance at which the message was generated.
-	 */
+	/** Holds the number of seconds from the start of the instance at which the message was generated. */
 	double TimeSeconds;
 
-	/**
-	 * Holds the verbosity type.
-	 */
+	/** Holds the verbosity type. */
 	ELogVerbosity::Type Verbosity;
 
+public:
 
 	/**
 	 * Creates and initializes a new instance.
 	 *
-	 * @param InInstanceId - The identifier of the instance that generated the log message.
-	 * @param InInstanceName - The name of the engine instance that generated the log message.
-	 * @param InTimeSeconds - The number of seconds from the start of the instance at which the message was generated.
-	 * @param InText - The message text.
-	 * @param InVerbosity - The verbosity type.
-	 * @param InCategory - The log category.
+	 * @param InInstanceId The identifier of the instance that generated the log message.
+	 * @param InInstanceName The name of the engine instance that generated the log message.
+	 * @param InTimeSeconds The number of seconds from the start of the instance at which the message was generated.
+	 * @param InText The message text.
+	 * @param InVerbosity The verbosity type.
+	 * @param InCategory The log category.
 	 */
 	FSessionLogMessage(const FGuid& InInstanceId, const FString& InInstanceName, float InTimeSeconds, const FString& InText, ELogVerbosity::Type InVerbosity, const FName& InCategory)
 		: Category(InCategory)
