Fix for available sessions not getting updated in UFE. It was testing FDateTime::UTCNow against FDateTime::Now so it would only get updated after an hour due to daylight saving.

diff --git a/Engine/Source/Developer/SessionServices/Private/SessionManager.cpp b/Engine/Source/Developer/SessionServices/Private/SessionManager.cpp
--- a/Engine/Source/Developer/SessionServices/Private/SessionManager.cpp
+++ b/Engine/Source/Developer/SessionServices/Private/SessionManager.cpp
@@ -301,7 +301,7 @@ void FSessionManager::HandleSessionPongMessage( const FSessionServicePong& Messa
 
 bool FSessionManager::HandleTicker( float DeltaTime )
 {
-	FDateTime Now = FDateTime::UtcNow();
+	FDateTime Now = FDateTime::Now();
 
 //	FindExpiredSessions(Now);
 
