SessionServices: Code and documentation cleanup pass

diff --git a/Engine/Source/Developer/SessionServices/Private/SessionInfo.h b/Engine/Source/Developer/SessionServices/Private/SessionInfo.h
--- a/Engine/Source/Developer/SessionServices/Private/SessionInfo.h
+++ b/Engine/Source/Developer/SessionServices/Private/SessionInfo.h
@@ -116,47 +116,47 @@ public:
 		}
 	}
 
-	virtual FDateTime GetLastUpdateTime( ) override
+	virtual FDateTime GetLastUpdateTime() override
 	{
 		return LastUpdateTime;
 	}
 
-	virtual const int32 GetNumInstances( ) const override
+	virtual const int32 GetNumInstances() const override
 	{
 		return Instances.Num();
 	}
 
-	virtual const FGuid& GetSessionId( ) const override
+	virtual const FGuid& GetSessionId() const override
 	{
 		return SessionId;
 	}
 
-	virtual const FString& GetSessionName( ) const override
+	virtual const FString& GetSessionName() const override
 	{
 		return SessionName;
 	}
 
-	virtual const FString& GetSessionOwner( ) const override
+	virtual const FString& GetSessionOwner() const override
 	{
 		return SessionOwner;
 	}
 
-	virtual const bool IsStandalone( ) const override
+	virtual const bool IsStandalone() const override
 	{
 		return Standalone;
 	}
 
-	virtual FOnSessionInstanceDiscovered& OnInstanceDiscovered( ) override
+	virtual FOnSessionInstanceDiscovered& OnInstanceDiscovered() override
 	{
 		return InstanceDiscoveredDelegate;
 	}
 
-	virtual FOnSessionLogReceived& OnLogReceived( ) override
+	virtual FOnSessionLogReceived& OnLogReceived() override
 	{
 		return LogReceivedDelegate;
 	}
 
-	virtual void Terminate( ) override
+	virtual void Terminate() override
 	{
 		for (TMap<FMessageAddress, TSharedPtr<FSessionInstanceInfo> >::TIterator It(Instances); It; ++It)
 		{
@@ -166,7 +166,7 @@ public:
 
 private:
 
-	// Handles received log messages.
+	/** Handles received log messages. */
 	void HandleLogReceived( const ISessionInstanceInfoRef& Instance, const FSessionLogMessageRef& LogMessage )
 	{
 		LogReceivedDelegate.Broadcast(AsShared(), Instance, LogMessage);
@@ -174,32 +174,32 @@ private:
 
 private:
 
-	// Holds the list of engine instances that belong to this session.
+	/** Holds the list of engine instances that belong to this session. */
 	TMap<FMessageAddress, TSharedPtr<FSessionInstanceInfo>> Instances;
 
-	// Holds the time at which the last pong was received.
+	/** Holds the time at which the last pong was received. */
 	FDateTime LastUpdateTime;
 
-	// Holds a weak pointer to the message bus.
+	/** Holds a weak pointer to the message bus. */
 	IMessageBusWeakPtr MessageBusPtr;
 
-	// Holds the session identifier.
+	/** Holds the session identifier. */
 	FGuid SessionId;
 
-	// Holds the session name.
+	/** Holds the session name. */
 	FString SessionName;
 
-	// Holds the name of the user who launched the session.
+	/** Holds the name of the user who launched the session. */
 	FString SessionOwner;
 
-	// Is the session local
+	/** Whether the session is local. */
 	bool Standalone;
 
 private:
 
-	// Holds a delegate to be invoked when a new instance has been discovered.
+	/** Holds a delegate to be invoked when a new instance has been discovered. */
 	FOnSessionInstanceDiscovered InstanceDiscoveredDelegate;
 
-	// Holds a delegate to be invoked when an instance received a log message.
+	/** Holds a delegate to be invoked when an instance received a log message. */
 	FOnSessionLogReceived LogReceivedDelegate;
 };
diff --git a/Engine/Source/Developer/SessionServices/Private/SessionInstanceInfo.h b/Engine/Source/Developer/SessionServices/Private/SessionInstanceInfo.h
--- a/Engine/Source/Developer/SessionServices/Private/SessionInstanceInfo.h
+++ b/Engine/Source/Developer/SessionServices/Private/SessionInstanceInfo.h
@@ -12,10 +12,8 @@ class FSessionInstanceInfo
 {
 public:
 
-	/**
-	 * Default constructor.
-	 */
-	FSessionInstanceInfo( ) { }
+	/** Default constructor. */
+	FSessionInstanceInfo() { }
 
 	/**
 	 * Creates and initializes a new instance.
@@ -111,12 +109,12 @@ public:
 		return DeviceName;
 	}
 
-	virtual int32 GetEngineVersion( ) const override
+	virtual int32 GetEngineVersion() const override
 	{
 		return EngineVersion;
 	}
 
-	virtual const FGuid GetInstanceId( ) const override
+	virtual const FGuid GetInstanceId() const override
 	{
 		return InstanceId;
 	}
@@ -131,27 +129,27 @@ public:
 		return InstanceType;
 	}
 
-	virtual FDateTime GetLastUpdateTime( ) override
+	virtual FDateTime GetLastUpdateTime() override
 	{
 		return LastUpdateTime;
 	}
 
-	virtual const TArray<FSessionLogMessagePtr>& GetLog( ) override
+	virtual const TArray<FSessionLogMessagePtr>& GetLog() override
 	{
 		return LogMessages;
 	}
 
-	virtual ISessionInfoPtr GetOwnerSession( ) override
+	virtual ISessionInfoPtr GetOwnerSession() override
 	{
 		return Owner.Pin();
 	}
 
-	virtual const FString& GetPlatformName( ) const override
+	virtual const FString& GetPlatformName() const override
 	{
 		return PlatformName;
 	}
 
-	virtual float GetWorldTimeSeconds( ) const override
+	virtual float GetWorldTimeSeconds() const override
 	{
 		return WorldTimeSeconds;
 	}
@@ -161,17 +159,17 @@ public:
 		return IsConsoleBuild;
 	}
 
-	virtual FOnSessionInstanceLogReceived& OnLogReceived( ) override
+	virtual FOnSessionInstanceLogReceived& OnLogReceived() override
 	{
 		return LogReceivedDelegate;
 	}
 
-	virtual bool PlayHasBegun( ) const override
+	virtual bool PlayHasBegun() const override
 	{
 		return HasBegunPlay;
 	}
 
-	virtual void Terminate( ) override
+	virtual void Terminate() override
 	{
 		if (MessageEndpoint.IsValid() && EngineAddress.IsValid())
 		{
@@ -201,59 +199,59 @@ private:
 
 private:
 
-	// Holds the message bus address of the application instance.
+	/** Holds the message bus address of the application instance. */
 	FMessageAddress ApplicationAddress;
 
-	// Holds the instance's build date.
+	/** Holds the instance's build date. */
 	FString BuildDate;
 
-	// Holds the instance's current level.
+	/** Holds the instance's current level. */
 	FString	CurrentLevel;
 
-	// Holds the device name.
+	/** Holds the device name. */
 	FString DeviceName;
 
-	// Holds the message bus address of the engine instance.
+	/** Holds the message bus address of the engine instance. */
 	FMessageAddress EngineAddress;
 
-	// Holds the instance's engine version.
+	/** Holds the instance's engine version. */
 	int32 EngineVersion;
 
-	// Holds a flag indicating whether the game has begun.
+	/** Holds a flag indicating whether the game has begun. */
 	bool HasBegunPlay;
 
-	// Holds the instance identifier.
+	/** Holds the instance identifier. */
 	FGuid InstanceId;
 
-	// Holds the instance name.
+	/** Holds the instance name. */
 	FString InstanceName;
 
-	// Holds the instance type (i.e. game, editor etc.)
+	/** Holds the instance type (i.e. game, editor etc.) */
 	FString InstanceType;
 
-	// Holds a flag indicating whether this is a console build.
+	/** Holds a flag indicating whether this is a console build. */
 	bool IsConsoleBuild;
 
-	// Holds the time at which the last pong was received.
+	/** Holds the time at which the last pong was received. */
 	FDateTime LastUpdateTime;
 
-	// Holds the collection of received log messages.
+	/** Holds the collection of received log messages. */
 	TArray<FSessionLogMessagePtr> LogMessages;
 
-	// Holds the message endpoint.
+	/** Holds the message endpoint. */
 	FMessageEndpointPtr MessageEndpoint;
 
-	// Holds a reference to the session that owns this instance.
+	/** Holds a reference to the session that owns this instance. */
 	TWeakPtr<ISessionInfo> Owner;
 
-	// Holds the name of the platform that the instance is running on.
+	/** Holds the name of the platform that the instance is running on. */
 	FString PlatformName;
 
-	// Holds the instance's current game world time.
+	/** Holds the instance's current game world time. */
 	float WorldTimeSeconds;
 
 private:
 
-	// Holds a delegate to be invoked when a log message was received from the remote session.
+	/** Holds a delegate to be invoked when a log message was received from the remote session. */
 	FOnSessionInstanceLogReceived LogReceivedDelegate;
 };
diff --git a/Engine/Source/Developer/SessionServices/Private/SessionManager.cpp b/Engine/Source/Developer/SessionServices/Private/SessionManager.cpp
--- a/Engine/Source/Developer/SessionServices/Private/SessionManager.cpp
+++ b/Engine/Source/Developer/SessionServices/Private/SessionManager.cpp
@@ -41,7 +41,7 @@ FSessionManager::FSessionManager( const IMessageBusRef& InMessageBus )
 }
 
 
-FSessionManager::~FSessionManager( )
+FSessionManager::~FSessionManager()
 {
 	FTicker::GetCoreTicker().RemoveTicker(TickDelegate);
 }
@@ -152,7 +152,7 @@ void FSessionManager::FindExpiredSessions( const FDateTime& Now )
 {
 	bool Dirty = false;
 
-	for (TMap<FMessageAddress, TSharedPtr<FSessionInfo> >::TIterator It(Sessions); It; ++It)
+	for (TMap<FGuid, TSharedPtr<FSessionInfo> >::TIterator It(Sessions); It; ++It)
 	{
 		if (Now > It.Value()->GetLastUpdateTime() + FTimespan::FromSeconds(10.0))
 		{
@@ -193,7 +193,7 @@ void FSessionManager::RefreshSessions()
 {
 	bool Dirty = false;
 
-	for (TMap<FMessageAddress, TSharedPtr<FSessionInfo> >::TIterator It(Sessions); It; ++It)
+	for (TMap<FGuid, TSharedPtr<FSessionInfo> >::TIterator It(Sessions); It; ++It)
 	{
 		if (!IsValidOwner(It.Value()->GetSessionOwner()))
 		{
@@ -210,7 +210,7 @@ void FSessionManager::RefreshSessions()
 }
 
 
-void FSessionManager::SendPing( )
+void FSessionManager::SendPing()
 {
 	if (MessageEndpoint.IsValid())
 	{
diff --git a/Engine/Source/Developer/SessionServices/Private/SessionManager.h b/Engine/Source/Developer/SessionServices/Private/SessionManager.h
--- a/Engine/Source/Developer/SessionServices/Private/SessionManager.h
+++ b/Engine/Source/Developer/SessionServices/Private/SessionManager.h
@@ -19,10 +19,8 @@ public:
 	 */
 	FSessionManager( const IMessageBusRef& InMessageBus );
 	
-	/**
-	 * Default destructor
-	 */
-	~FSessionManager( );
+	/** Default destructor */
+	~FSessionManager();
 
 public:
 
@@ -31,7 +29,7 @@ public:
 	virtual void AddOwner( const FString& InOwner ) override;
 	virtual void GetSelectedInstances( TArray<ISessionInstanceInfoPtr>& OutInstances) const override;
 
-	virtual const ISessionInfoPtr& GetSelectedSession( ) const override
+	virtual const ISessionInfoPtr& GetSelectedSession() const override
 	{
 		return SelectedSession;
 	}
@@ -43,32 +41,32 @@ public:
 		return ((Instance->GetOwnerSession() == SelectedSession) && !DeselectedInstances.Contains(Instance));
 	}
 
-	virtual FOnCanSelectSession& OnCanSelectSession( ) override
+	virtual FOnCanSelectSession& OnCanSelectSession() override
 	{
 		return CanSelectSessionDelegate;
 	}
 
-	virtual FOnSessionInstanceSelectionChanged& OnInstanceSelectionChanged( ) override
+	virtual FOnSessionInstanceSelectionChanged& OnInstanceSelectionChanged() override
 	{
 		return InstanceSelectionChangedDelegate;
 	}
 
-	virtual FOnSessionLogReceived& OnLogReceived( ) override
+	virtual FOnSessionLogReceived& OnLogReceived() override
 	{
 		return LogReceivedDelegate;
 	}
 
-	virtual FOnSelectedSessionChanged& OnSelectedSessionChanged( ) override
+	virtual FOnSelectedSessionChanged& OnSelectedSessionChanged() override
 	{
 		return SelectedSessionChangedDelegate;
 	}
 
-	virtual FSimpleMulticastDelegate& OnSessionsUpdated( ) override
+	virtual FSimpleMulticastDelegate& OnSessionsUpdated() override
 	{
 		return SessionsUpdatedDelegate;
 	}
 
-	virtual FSimpleMulticastDelegate& OnSessionInstanceUpdated( ) override
+	virtual FSimpleMulticastDelegate& OnSessionInstanceUpdated() override
 	{
 		return SessionInstanceUpdatedDelegate;
 	}
@@ -93,79 +91,75 @@ protected:
 	 */
 	bool IsValidOwner( const FString& Owner );
 
-	/**
-	 * Refresh the sessions based on the owner filter list.
-	 */
+	/** Refresh the sessions based on the owner filter list. */
 	void RefreshSessions();
 
-	/**
-	 * Pings all sessions on the network.
-	 */
-	void SendPing( );
+	/** Pings all sessions on the network. */
+	void SendPing();
 
 private:
 
-	// Callback for handling FSessionServicePong messages.
+	/** Callback for handling FSessionServicePong messages. */
 	void HandleEnginePongMessage( const FEngineServicePong& Message, const IMessageContextRef& Context );
 
-	// Callback for newly discovered instances.
+	/** Callback for newly discovered instances. */
 	void HandleInstanceDiscovered( const ISessionInfoRef& Session, const ISessionInstanceInfoRef& Instance );
 
-	// Callback received log entries.
+	/** Callback received log entries. */
 	void HandleLogReceived( const ISessionInfoRef& Session, const ISessionInstanceInfoRef& Instance, const FSessionLogMessageRef& Message );
 
-	// Callback for handling FSessionServicePong messages.
+	/** Callback for handling FSessionServicePong messages. */
 	void HandleSessionPongMessage( const FSessionServicePong& Message, const IMessageContextRef& Context );
 
-	// Callback for ticks from the ticker.
+	/** Callback for ticks from the ticker. */
 	bool HandleTicker( float DeltaTime );
 
 private:
 
-	// The address of the automation controller to where we can forward any automation workers found.
+	/** The address of the automation controller to where we can forward any automation workers found. */
 	FGuid AutomationControllerAddress;
 
-	// Holds the list of currently selected instances.
+	/** Holds the list of currently selected instances. */
 	TArray<ISessionInstanceInfoPtr> DeselectedInstances;
 
-	// Holds the time at which the last ping was sent.
+	/** Holds the time at which the last ping was sent. */
 	FDateTime LastPingTime;
 
-	// Holds a pointer to the message bus.
+	/** Holds a pointer to the message bus. */
 	IMessageBusWeakPtr MessageBusPtr;
 
-	// Holds the messaging endpoint.
+	/** Holds the messaging endpoint. */
 	FMessageEndpointPtr MessageEndpoint;
 
-	// Holds a reference to the currently selected session.
+	/** Holds a reference to the currently selected session. */
 	ISessionInfoPtr SelectedSession;
 
-	// Holds the collection of discovered sessions.
+	/** Holds the collection of discovered sessions. */
 	TMap<FGuid, TSharedPtr<FSessionInfo> > Sessions;
 
 private:
 
-	// Holds a delegate to be invoked before a session is selected.
+	/** Holds a delegate to be invoked before a session is selected. */
 	FOnCanSelectSession CanSelectSessionDelegate;
 
-	// Holds a delegate to be invoked when an instance changes its selection state.
+	/** Holds a delegate to be invoked when an instance changes its selection state. */
 	FOnSessionInstanceSelectionChanged InstanceSelectionChangedDelegate;
 
-	// Holds a delegate to be invoked when the selected session received a log message.
+	/** Owner filter list. */
+	TArray<FString> FilteredOwners;
+
+	/** Holds a delegate to be invoked when the selected session received a log message. */
 	FOnSessionLogReceived LogReceivedDelegate;
 
-	// Holds a delegate to be invoked when the selected session changed.
+	/** Holds a delegate to be invoked when the selected session changed. */
 	FOnSelectedSessionChanged SelectedSessionChangedDelegate;
 
-	// Holds a delegate to be invoked when the session list was updated.
+	/** Holds a delegate to be invoked when the session list was updated. */
 	FSimpleMulticastDelegate SessionsUpdatedDelegate;
 
-	// Holds a delegate to be invoked when a session instance is updated.
+	/** Holds a delegate to be invoked when a session instance is updated. */
 	FSimpleMulticastDelegate SessionInstanceUpdatedDelegate;
 
-	// Holds a delegate to be invoked when the widget ticks.
+	/** Holds a delegate to be invoked when the widget ticks. */
 	FTickerDelegate TickDelegate;
-
-	// Owner filter list
-	TArray<FString> FilteredOwners;
 };
diff --git a/Engine/Source/Developer/SessionServices/Private/SessionService.h b/Engine/Source/Developer/SessionServices/Private/SessionService.h
--- a/Engine/Source/Developer/SessionServices/Private/SessionService.h
+++ b/Engine/Source/Developer/SessionServices/Private/SessionService.h
@@ -19,10 +19,8 @@ public:
 	 */
 	FSessionService( const IMessageBusRef& InMessageBus );
 
-	/**
-	 * Destructor.
-	 */
-	~FSessionService( );
+	/** Destructor. */
+	~FSessionService();
 
 public:
 
@@ -37,13 +35,13 @@ public:
 
 	// ISessionService interface
 
-	virtual bool IsRunning( ) override
+	virtual bool IsRunning() override
 	{
 		return MessageEndpoint.IsValid();
 	}
 
-	virtual bool Start( ) override;
-	virtual void Stop( ) override;
+	virtual bool Start() override;
+	virtual void Stop() override;
 
 protected:
 
@@ -73,29 +71,29 @@ protected:
 
 private:
 
-	// Handles message bus shutdowns.
-	void HandleMessageEndpointShutdown( );
+	/** Handles message bus shutdowns. */
+	void HandleMessageEndpointShutdown();
 
-	// Handles FSessionServiceLogSubscribe messages.
+	/** Handles FSessionServiceLogSubscribe messages. */
 	void HandleSessionLogSubscribeMessage( const FSessionServiceLogSubscribe& Message, const IMessageContextRef& Context );
 
-	// Handles FSessionServiceLogUnsubscribe messages.
+	/** Handles FSessionServiceLogUnsubscribe messages. */
 	void HandleSessionLogUnsubscribeMessage( const FSessionServiceLogUnsubscribe& Message, const IMessageContextRef& Context );
 
-	// Handles FSessionServicePing messages.
+	/** Handles FSessionServicePing messages. */
 	void HandleSessionPingMessage( const FSessionServicePing& Message, const IMessageContextRef& Context );
 
 private:
 
-	// Holds the list of log subscribers.
+	/** Holds the list of log subscribers. */
 	TArray<FMessageAddress> LogSubscribers;
 
-	// Holds a critical section for the log subscribers array.
+	/** Holds a critical section for the log subscribers array. */
 	FCriticalSection LogSubscribersLock;
 
-	// Holds a weak pointer to the message bus.
+	/** Holds a weak pointer to the message bus. */
 	IMessageBusWeakPtr MessageBusPtr;
 
-	// Holds the message endpoint.
+	/** Holds the message endpoint. */
 	FMessageEndpointPtr MessageEndpoint;
 };
diff --git a/Engine/Source/Developer/SessionServices/Private/SessionServicesModule.cpp b/Engine/Source/Developer/SessionServices/Private/SessionServicesModule.cpp
--- a/Engine/Source/Developer/SessionServices/Private/SessionServicesModule.cpp
+++ b/Engine/Source/Developer/SessionServices/Private/SessionServicesModule.cpp
@@ -14,19 +14,17 @@ class FSessionServicesModule
 {
 public:
 
-	/**
-	 * Default constructor.
-	 */
-	FSessionServicesModule( )
-		: SessionManager(NULL)
-		, SessionService(NULL)
+	/** Default constructor. */
+	FSessionServicesModule()
+		: SessionManager(nullptr)
+		, SessionService(nullptr)
 	{ }
 
 public:
 
 	// ISessionServicesModule interface
 
-	virtual ISessionManagerRef GetSessionManager( ) override
+	virtual ISessionManagerRef GetSessionManager() override
 	{
 		if (!SessionManager.IsValid())
 		{
@@ -36,7 +34,7 @@ public:
 		return SessionManager.ToSharedRef();
 	}
 
-	virtual ISessionServiceRef GetSessionService( ) override
+	virtual ISessionServiceRef GetSessionService() override
 	{
 		if (!SessionService.IsValid())
 		{
@@ -50,13 +48,13 @@ public:
 
 	// IModuleInterface interface
 
-	virtual void StartupModule( ) override
+	virtual void StartupModule() override
 	{
 		MessageBusPtr = IMessagingModule::Get().GetDefaultBus();
 		check(MessageBusPtr.IsValid());
 	}
 
-	virtual void ShutdownModule( ) override
+	virtual void ShutdownModule() override
 	{
 		SessionManager.Reset();
 		SessionService.Reset();
@@ -64,13 +62,13 @@ public:
 
 private:
 	
-	// Holds a weak pointer to the message bus.
+	/** Holds a weak pointer to the message bus. */
 	IMessageBusWeakPtr MessageBusPtr;
 
-	// Holds the session manager singleton.
+	/** Holds the session manager singleton. */
 	ISessionManagerPtr SessionManager;
 
-	// Holds the session service singleton.
+	/** Holds the session service singleton. */
 	ISessionServicePtr SessionService;
 };
 
diff --git a/Engine/Source/Developer/SessionServices/Private/SessionServicesPrivatePCH.h b/Engine/Source/Developer/SessionServices/Private/SessionServicesPrivatePCH.h
--- a/Engine/Source/Developer/SessionServices/Private/SessionServicesPrivatePCH.h
+++ b/Engine/Source/Developer/SessionServices/Private/SessionServicesPrivatePCH.h
@@ -8,7 +8,7 @@
 /* Dependencies
  *****************************************************************************/
 
-#include "EngineMessages.h"
+#include "EngineServiceMessages.h"
 #include "Messaging.h"
 #include "SessionMessages.h"
 #include "Ticker.h"
diff --git a/Engine/Source/Developer/SessionServices/Public/Interfaces/ISessionInfo.h b/Engine/Source/Developer/SessionServices/Public/Interfaces/ISessionInfo.h
--- a/Engine/Source/Developer/SessionServices/Public/Interfaces/ISessionInfo.h
+++ b/Engine/Source/Developer/SessionServices/Public/Interfaces/ISessionInfo.h
@@ -47,35 +47,35 @@ public:
 	 *
 	 * @return The receive time.
 	 */
-	virtual FDateTime GetLastUpdateTime( ) = 0;
+	virtual FDateTime GetLastUpdateTime() = 0;
 
 	/**
 	 * Gets the number of engine instances that are part of the session.
 	 *
 	 * @return The number of engine instances.
 	 */
-	virtual const int32 GetNumInstances( ) const = 0;
+	virtual const int32 GetNumInstances() const = 0;
 
 	/**
 	 * Gets the session identifier.
 	 *
 	 * @return Session identifier.
 	 */
-	virtual const FGuid& GetSessionId( ) const = 0;
+	virtual const FGuid& GetSessionId() const = 0;
 
 	/**
 	 * Gets the name of the session.
 	 *
 	 * @return Session name.
 	 */
-	virtual const FString& GetSessionName( ) const = 0;
+	virtual const FString& GetSessionName() const = 0;
 
 	/**
 	 * Gets the name of the user that owns the session.
 	 *
 	 * @return User name.
 	 */
-	virtual const FString& GetSessionOwner( ) const = 0;
+	virtual const FString& GetSessionOwner() const = 0;
 
 	/**
 	 * Checks whether this is a standalone session.
@@ -84,12 +84,12 @@ public:
 	 *
 	 * @return true if this is a standalone session, false otherwise.
 	 */
-	virtual const bool IsStandalone( ) const = 0;
+	virtual const bool IsStandalone() const = 0;
 
 	/**
 	 * Terminates the session.
 	 */
-	virtual void Terminate( ) = 0;
+	virtual void Terminate() = 0;
 
 public:
 
@@ -98,19 +98,17 @@ public:
 	 *
 	 * @return The delegate.
 	 */
-	virtual FOnSessionInstanceDiscovered& OnInstanceDiscovered( ) = 0;
+	virtual FOnSessionInstanceDiscovered& OnInstanceDiscovered() = 0;
 
 	/**
 	 * Returns a delegate that is executed when a new log message has been received.
 	 *
 	 * @return The delegate.
 	 */
-	virtual FOnSessionLogReceived& OnLogReceived( ) = 0;
+	virtual FOnSessionLogReceived& OnLogReceived() = 0;
 
 public:
 
-	/**
-	 * Virtual destructor.
-	 */
-	virtual ~ISessionInfo( ) { }
+	/** Virtual destructor. */
+	virtual ~ISessionInfo() { }
 };
diff --git a/Engine/Source/Developer/SessionServices/Public/Interfaces/ISessionInstanceInfo.h b/Engine/Source/Developer/SessionServices/Public/Interfaces/ISessionInstanceInfo.h
--- a/Engine/Source/Developer/SessionServices/Public/Interfaces/ISessionInstanceInfo.h
+++ b/Engine/Source/Developer/SessionServices/Public/Interfaces/ISessionInstanceInfo.h
@@ -38,103 +38,103 @@ public:
 	 *
 	 * @return The build date string.
 	 */
-	virtual const FString GetBuildDate( ) const = 0;
+	virtual const FString GetBuildDate() const = 0;
 
 	/**
 	 * Gets the name of the level that the instance is currently running.
 	 *
 	 * @return The name of the level.
 	 */
-	virtual const FString GetCurrentLevel( ) const = 0;
+	virtual const FString GetCurrentLevel() const = 0;
 
 	/**
 	 * Gets the name of the device that this instance is running on.
 	 *
 	 * @return The device name string.
 	 */
-	virtual const FString GetDeviceName( ) const = 0;
+	virtual const FString GetDeviceName() const = 0;
 
 	/**
 	 * Gets the instance's engine version number.
 	 *
 	 * @return Engine version.
 	 */
-	virtual int32 GetEngineVersion( ) const = 0;
+	virtual int32 GetEngineVersion() const = 0;
 
 	/**
 	 * Gets the instance identifier.
 	 *
 	 * @return Instance identifier.
 	 */
-	virtual const FGuid GetInstanceId( ) const = 0;
+	virtual const FGuid GetInstanceId() const = 0;
 
 	/**
 	 * Gets the name of this instance.
 	 *
 	 * @return The instance name string.
 	 */
-	virtual const FString GetInstanceName( ) const = 0;
+	virtual const FString GetInstanceName() const = 0;
 
 	/**
 	 * Gets the instance type (i.e. Editor or Game).
 	 *
 	 * @return The game instance type string.
 	 */
-	virtual const FString GetInstanceType( ) const = 0;
+	virtual const FString GetInstanceType() const = 0;
 
 	/**
 	 * Gets the time at which the last update was received from this instance.
 	 *
 	 * @return The receive time.
 	 */
-	virtual FDateTime GetLastUpdateTime( ) = 0;
+	virtual FDateTime GetLastUpdateTime() = 0;
 
 	/**
 	 * Gets the collection of log entries received from this instance.
 	 *
 	 * @return Log entries.
 	 */
-	virtual const TArray<FSessionLogMessagePtr>& GetLog( ) = 0;
+	virtual const TArray<FSessionLogMessagePtr>& GetLog() = 0;
 
 	/**
 	 * Gets a reference to the session that owns this instance.
 	 *
 	 * @return Owner session.
 	 */
-	virtual TSharedPtr<class ISessionInfo> GetOwnerSession( ) = 0;
+	virtual TSharedPtr<class ISessionInfo> GetOwnerSession() = 0;
 
 	/**
 	 * Gets the name of the platform that the instance is running on.
 	 *
 	 * @return Platform name string.
 	 */
-	virtual const FString& GetPlatformName( ) const = 0;
+	virtual const FString& GetPlatformName() const = 0;
 
 	/**
 	 * Gets the instance's current game world time.
 	 *
 	 * @return World time in seconds.
 	 */
-	virtual float GetWorldTimeSeconds( ) const = 0;
+	virtual float GetWorldTimeSeconds() const = 0;
 
 	/**
 	 * Checks whether this instance is a console build (i.e. no editor features).
 	 *
 	 * @return true if it is a console build, false otherwise.
 	 */
-	virtual const bool IsConsole( ) const = 0;
+	virtual const bool IsConsole() const = 0;
 
 	/**
 	 * Checks whether this instance has already begun game play.
 	 *
 	 * @return true if game play has begun, false otherwise.
 	 */
-	virtual bool PlayHasBegun( ) const = 0;
+	virtual bool PlayHasBegun() const = 0;
 
 	/**
 	 * Terminates the instance.
 	 */
-	virtual void Terminate( ) = 0;
+	virtual void Terminate() = 0;
 
 public:
 
@@ -143,12 +143,10 @@ public:
 	 *
 	 * @return The delegate.
 	 */
-	virtual FOnSessionInstanceLogReceived& OnLogReceived( ) = 0;
+	virtual FOnSessionInstanceLogReceived& OnLogReceived() = 0;
 
 public:
 
-	/**
-	 * Virtual destructor.
-	 */
-	virtual ~ISessionInstanceInfo( ) { }
+	/** Virtual destructor. */
+	virtual ~ISessionInstanceInfo() { }
 };
diff --git a/Engine/Source/Developer/SessionServices/Public/Interfaces/ISessionManager.h b/Engine/Source/Developer/SessionServices/Public/Interfaces/ISessionManager.h
--- a/Engine/Source/Developer/SessionServices/Public/Interfaces/ISessionManager.h
+++ b/Engine/Source/Developer/SessionServices/Public/Interfaces/ISessionManager.h
@@ -3,13 +3,6 @@
 #pragma once
 
 
-/** Type definition for shared pointers to instances of ISessionManager. */
-typedef TSharedPtr<class ISessionManager> ISessionManagerPtr;
-
-/** Type definition for shared references to instances of ISessionManager. */
-typedef TSharedRef<class ISessionManager> ISessionManagerRef;
-
-
 /**
  * Delegate type for determining whether a session can be selected.
  *
@@ -59,7 +52,7 @@ public:
 	 *
 	 * @return The session ID selected
 	 */
-	virtual const ISessionInfoPtr& GetSelectedSession( ) const = 0;
+	virtual const ISessionInfoPtr& GetSelectedSession() const = 0;
 
 	/**
 	 * Gets the list of all discovered sessions.
@@ -106,47 +99,52 @@ public:
 	 *
 	 * @return The delegate.
 	 */
-	virtual FOnCanSelectSession& OnCanSelectSession( ) = 0;
+	virtual FOnCanSelectSession& OnCanSelectSession() = 0;
 
 	/**
 	 * Returns a delegate that is executed when an instance changes its selection state.
 	 *
 	 * @return The delegate.
 	 */
-	virtual FOnSessionInstanceSelectionChanged& OnInstanceSelectionChanged( ) = 0;
+	virtual FOnSessionInstanceSelectionChanged& OnInstanceSelectionChanged() = 0;
 
 	/**
 	 * Returns a delegate that is executed when the selected session received a log message from one of its instances.
 	 *
 	 * @return The delegate.
 	 */
-	virtual FOnSessionLogReceived& OnLogReceived( ) = 0;
+	virtual FOnSessionLogReceived& OnLogReceived() = 0;
 
 	/**
 	 * Returns a delegate that is executed when the selected session changed.
 	 *
 	 * @return The delegate.
 	 */
-	virtual FOnSelectedSessionChanged& OnSelectedSessionChanged( ) = 0;
+	virtual FOnSelectedSessionChanged& OnSelectedSessionChanged() = 0;
 
 	/**
 	 * Returns a delegate that is executed when the list of sessions has changed.
 	 *
 	 * @return The delegate.
 	 */
-	virtual FSimpleMulticastDelegate& OnSessionsUpdated( ) = 0;
+	virtual FSimpleMulticastDelegate& OnSessionsUpdated() = 0;
 
 	/**
 	 * Returns a delegate that is executed when a session instance is updated.
 	 *
 	 * @return The delegate.
 	 */
-	virtual FSimpleMulticastDelegate& OnSessionInstanceUpdated( ) = 0;
+	virtual FSimpleMulticastDelegate& OnSessionInstanceUpdated() = 0;
 
 public:
 
-	/**
-	 * Virtual destructor.
-	 */
-	virtual ~ISessionManager( ) { }
+	/** Virtual destructor. */
+	virtual ~ISessionManager() { }
 };
+
+
+/** Type definition for shared pointers to instances of ISessionManager. */
+typedef TSharedPtr<ISessionManager> ISessionManagerPtr;
+
+/** Type definition for shared references to instances of ISessionManager. */
+typedef TSharedRef<ISessionManager> ISessionManagerRef;
diff --git a/Engine/Source/Developer/SessionServices/Public/Interfaces/ISessionService.h b/Engine/Source/Developer/SessionServices/Public/Interfaces/ISessionService.h
--- a/Engine/Source/Developer/SessionServices/Public/Interfaces/ISessionService.h
+++ b/Engine/Source/Developer/SessionServices/Public/Interfaces/ISessionService.h
@@ -22,7 +22,7 @@ public:
 	 *
 	 * @return true if the service is running, false otherwise.
 	 */
-	virtual bool IsRunning( ) = 0;
+	virtual bool IsRunning() = 0;
 
 	/**
 	 * Starts the service.
@@ -30,18 +30,16 @@ public:
 	 * @return true if the service was started, false otherwise.
 	 * @see Stop
 	 */
-	virtual bool Start( ) = 0;
+	virtual bool Start() = 0;
 
 	/**
 	 * Stops the service.
 	 * @see Start
 	 */
-	virtual void Stop( ) = 0;
+	virtual void Stop() = 0;
 
 public:
 
-	/**
-	 * Virtual destructor.
-	 */
-	virtual ~ISessionService( ) { }
+	/** Virtual destructor. */
+	virtual ~ISessionService() { }
 };
diff --git a/Engine/Source/Developer/SessionServices/Public/Interfaces/ISessionServicesModule.h b/Engine/Source/Developer/SessionServices/Public/Interfaces/ISessionServicesModule.h
--- a/Engine/Source/Developer/SessionServices/Public/Interfaces/ISessionServicesModule.h
+++ b/Engine/Source/Developer/SessionServices/Public/Interfaces/ISessionServicesModule.h
@@ -16,19 +16,17 @@ public:
 	 *
 	 * @return The session manager.
 	 */
-	virtual ISessionManagerRef GetSessionManager( ) = 0;
+	virtual ISessionManagerRef GetSessionManager() = 0;
 
 	/** 
 	 * Gets the session service.
 	 *
 	 * @return The session service.
 	 */
-	virtual ISessionServiceRef GetSessionService( ) = 0;
+	virtual ISessionServiceRef GetSessionService() = 0;
 
 public:
 
-	/**
-	 * Virtual destructor.
-	 */
-	virtual ~ISessionServicesModule( ) { }
+	/** Virtual destructor. */
+	virtual ~ISessionServicesModule() { }
 };
