Added missing include

diff --git a/Engine/Source/Developer/SessionServices/Private/SessionServicesPrivatePCH.h b/Engine/Source/Developer/SessionServices/Private/SessionServicesPrivatePCH.h
--- a/Engine/Source/Developer/SessionServices/Private/SessionServicesPrivatePCH.h
+++ b/Engine/Source/Developer/SessionServices/Private/SessionServicesPrivatePCH.h
@@ -8,6 +8,7 @@
 /* Dependencies
  *****************************************************************************/
 
+#include "CoreUObject.h"
 #include "EngineServiceMessages.h"
 #include "Messaging.h"
 #include "SessionMessages.h"
