Rollback deleted SessionServices files tilll we can fix merge issue

diff --git a/Engine/Source/Developer/SessionServices/Public/ISessionInfo.h b/Engine/Source/Developer/SessionServices/Public/ISessionInfo.h
new file mode 100644
--- /dev/null
+++ b/Engine/Source/Developer/SessionServices/Public/ISessionInfo.h
@@ -0,0 +1,102 @@
+// Copyright 1998-2015 Epic Games, Inc. All Rights Reserved.
+
+#pragma once
+
+
+class ISessionInstanceInfo;
+struct FSessionLogMessage;
+
+
+/**
+ * Interface for game instances.
+ */
+class ISessionInfo
+{
+public:
+
+	/**
+	 * Gets a read-only collection of all instances that belong to this session.
+	 *
+	 * @param OutInstances Will hold the collection of instances.
+	 */
+	virtual void GetInstances( TArray<TSharedPtr<ISessionInstanceInfo>>& OutInstances ) const = 0;
+
+	/**
+	 * Gets the time at which the last update was received from this instance.
+	 *
+	 * @return The receive time.
+	 */
+	virtual const FDateTime& GetLastUpdateTime() const = 0;
+
+	/**
+	 * Gets the number of engine instances that are part of the session.
+	 *
+	 * @return The number of engine instances.
+	 */
+	virtual const int32 GetNumInstances() const = 0;
+
+	/**
+	 * Gets the session identifier.
+	 *
+	 * @return Session identifier.
+	 */
+	virtual const FGuid& GetSessionId() const = 0;
+
+	/**
+	 * Gets the name of the session.
+	 *
+	 * @return Session name.
+	 */
+	virtual const FString& GetSessionName() const = 0;
+
+	/**
+	 * Gets the name of the user that owns the session.
+	 *
+	 * @return User name.
+	 */
+	virtual const FString& GetSessionOwner() const = 0;
+
+	/**
+	 * Checks whether this is a standalone session.
+	 *
+	 * A session is standalone if has not been created from the Launcher.
+	 *
+	 * @return true if this is a standalone session, false otherwise.
+	 */
+	virtual const bool IsStandalone() const = 0;
+
+	/**
+	 * Terminates the session.
+	 */
+	virtual void Terminate() = 0;
+
+public:
+
+	/**
+	 * A delegate that is executed when a new instance has been discovered.
+	 *
+	 * @return The delegate.
+	 */
+	DECLARE_EVENT_TwoParams(ISessionInfo, FInstanceDiscoveredEvent, const TSharedRef<ISessionInfo>& /*OwnerSession*/, const TSharedRef<ISessionInstanceInfo>& /*DiscoveredInstance*/)
+	virtual FInstanceDiscoveredEvent& OnInstanceDiscovered() = 0;
+
+	/**
+	 * A delegate that is executed when a new log message has been received.
+	 *
+	 * @return The delegate.
+	 */
+	DECLARE_EVENT_ThreeParams(ISessionInfo, FLogReceivedEvent, const TSharedRef<ISessionInfo>& /*OwnerSession*/, const TSharedRef<ISessionInstanceInfo>& /*SessionInstance*/, const TSharedRef<FSessionLogMessage>& /*LogMessage*/);
+	virtual FLogReceivedEvent& OnLogReceived() = 0;
+
+public:
+
+	/** Virtual destructor. */
+	virtual ~ISessionInfo() { }
+};
+
+
+/** Type definition for shared pointers to instances of ISessionInfo. */
+typedef TSharedPtr<ISessionInfo> ISessionInfoPtr;
+
+/** Type definition for shared references to instances of ISessionInfo. */
+typedef TSharedRef<ISessionInfo> ISessionInfoRef;
diff --git a/Engine/Source/Developer/SessionServices/Public/ISessionInstanceInfo.h b/Engine/Source/Developer/SessionServices/Public/ISessionInstanceInfo.h
new file mode 100644
--- /dev/null
+++ b/Engine/Source/Developer/SessionServices/Public/ISessionInstanceInfo.h
@@ -0,0 +1,146 @@
+// Copyright 1998-2015 Epic Games, Inc. All Rights Reserved.
+
+#pragma once
+
+
+class ISessionInfo;
+struct FSessionLogMessage;
+
+
+/**
+ * Interface for game instances.
+ */
+class ISessionInstanceInfo
+{
+public:
+
+	/**
+	 * Executes a console command on this engine instance.
+	 *
+	 * @param CommandString The command to execute.
+	 */
+	virtual void ExecuteCommand( const FString& CommandString ) = 0;
+
+	/**
+	 * Gets instance's build date.
+	 *
+	 * @return The build date string.
+	 */
+	virtual const FString& GetBuildDate() const = 0;
+
+	/**
+	 * Gets the name of the level that the instance is currently running.
+	 *
+	 * @return The name of the level.
+	 */
+	virtual const FString& GetCurrentLevel() const = 0;
+
+	/**
+	 * Gets the name of the device that this instance is running on.
+	 *
+	 * @return The device name string.
+	 */
+	virtual const FString& GetDeviceName() const = 0;
+
+	/**
+	 * Gets the instance's engine version number.
+	 *
+	 * @return Engine version.
+	 */
+	virtual int32 GetEngineVersion() const = 0;
+
+	/**
+	 * Gets the instance identifier.
+	 *
+	 * @return Instance identifier.
+	 */
+	virtual const FGuid& GetInstanceId() const = 0;
+
+	/**
+	 * Gets the name of this instance.
+	 *
+	 * @return The instance name string.
+	 */
+	virtual const FString& GetInstanceName() const = 0;
+
+	/**
+	 * Gets the instance type (i.e. Editor or Game).
+	 *
+	 * @return The game instance type string.
+	 */
+	virtual const FString& GetInstanceType() const = 0;
+
+	/**
+	 * Gets the time at which the last update was received from this instance.
+	 *
+	 * @return The receive time.
+	 */
+	virtual const FDateTime& GetLastUpdateTime() const = 0;
+
+	/**
+	 * Gets the collection of log entries received from this instance.
+	 *
+	 * @return Log entries.
+	 */
+	virtual const TArray<TSharedPtr<FSessionLogMessage>>& GetLog() = 0;
+
+	/**
+	 * Gets a reference to the session that owns this instance.
+	 *
+	 * @return Owner session.
+	 */
+	virtual TSharedPtr<ISessionInfo> GetOwnerSession() = 0;
+
+	/**
+	 * Gets the name of the platform that the instance is running on.
+	 *
+	 * @return Platform name string.
+	 */
+	virtual const FString& GetPlatformName() const = 0;
+
+	/**
+	 * Gets the instance's current game world time.
+	 *
+	 * @return World time in seconds.
+	 */
+	virtual float GetWorldTimeSeconds() const = 0;
+
+	/**
+	 * Checks whether this instance is a console build (i.e. no editor features).
+	 *
+	 * @return true if it is a console build, false otherwise.
+	 */
+	virtual const bool IsConsole() const = 0;
+
+	/**
+	 * Checks whether this instance has already begun game play.
+	 *
+	 * @return true if game play has begun, false otherwise.
+	 */
+	virtual bool PlayHasBegun() const = 0;
+
+	/** Terminates the instance. */
+	virtual void Terminate() = 0;
+
+public:
+
+	/**
+	 * Gets an event delegate that is executed when a new log message has been received.
+	 *
+	 * @return The delegate.
+	 */
+	DECLARE_EVENT_TwoParams(ISessionInstanceInfo, FLogReceivedEvent, const TSharedRef<ISessionInstanceInfo>&, const TSharedRef<FSessionLogMessage>&);
+	virtual FLogReceivedEvent& OnLogReceived() = 0;
+
+public:
+
+	/** Virtual destructor. */
+	virtual ~ISessionInstanceInfo() { }
+};
+
+
+/** Type definition for shared pointers to instances of IGameInstanceInfo. */
+typedef TSharedPtr<ISessionInstanceInfo> ISessionInstanceInfoPtr;
+
+/** Type definition for shared references to instances of IGameInstanceInfo. */
+typedef TSharedRef<ISessionInstanceInfo> ISessionInstanceInfoRef;
diff --git a/Engine/Source/Developer/SessionServices/Public/ISessionManager.h b/Engine/Source/Developer/SessionServices/Public/ISessionManager.h
new file mode 100644
--- /dev/null
+++ b/Engine/Source/Developer/SessionServices/Public/ISessionManager.h
@@ -0,0 +1,135 @@
+// Copyright 1998-2015 Epic Games, Inc. All Rights Reserved.
+
+#pragma once
+
+
+class ISessionInfo;
+class ISessionInstanceInfo;
+
+
+/**
+ * Interface for the session manager.
+ */
+class ISessionManager
+{
+public:
+
+	/**
+	 * Adds an owner whose sessions we are interested in
+	 *
+	 * @param InOwner Session owner we want to view sessions from
+	 */
+	virtual void AddOwner( const FString& InOwner ) = 0;
+
+	/**
+	  * Gets the collection of currently selected engine instances.
+	  *
+	  * @param OutInstances Will hold the selected instances.
+	  */
+	virtual void GetSelectedInstances( TArray<TSharedPtr<ISessionInstanceInfo>>& OutInstances) const = 0;
+
+	/**
+	 * Get the selected session - as chosen in the session browser
+	 *
+	 * @return The session ID selected
+	 */
+	virtual const TSharedPtr<ISessionInfo>& GetSelectedSession() const = 0;
+
+	/**
+	 * Gets the list of all discovered sessions.
+	 *
+	 * @param OutSessions Will hold the collection of sessions.
+	 */
+	virtual void GetSessions( TArray<TSharedPtr<ISessionInfo>>& OutSessions ) const = 0;
+
+	/**
+	 * Checks whether the given instance is currently selected.
+	 *
+	 * @return true if the instance is selected, false otherwise.
+	 */
+	virtual bool IsInstanceSelected( const TSharedRef<ISessionInstanceInfo>& Instance ) const = 0;
+
+	/**
+	 * Removes an owner whose sessions we are no longer interested in
+	 *
+	 * @param InOwner Session owner we want to remove
+	 */
+	virtual void RemoveOwner( const FString& InOwner ) = 0;
+
+	/**
+	 * Selects the specified session.
+	 *
+	 * @param Session The session to the select (can be NULL to select none).
+	 * @return true if the session was selected, false otherwise.
+	 */
+	virtual bool SelectSession( const TSharedPtr<ISessionInfo>& Session ) = 0;
+
+	/**
+	 * Marks the specified item as selected or unselected.
+	 *
+	 * @param Instance The instance to mark.
+	 * @param Selected Whether the instance should be selected (true) or unselected (false).
+	 * @return true if the instance was selected, false otherwise.
+	 */
+	virtual bool SetInstanceSelected( const TSharedPtr<ISessionInstanceInfo>& Instance, bool Selected ) = 0;
+
+public:
+
+	/**
+	 * Returns a delegate that is executed before a session is being selected.
+	 *
+	 * @return The delegate.
+	 */
+	DECLARE_EVENT_TwoParams(ISessionManager, FCanSelectSessionEvent, const TSharedPtr<ISessionInfo>& /*Session*/, bool& /*OutCanSelect*/)
+	virtual FCanSelectSessionEvent& OnCanSelectSession() = 0;
+
+	/**
+	 * Returns a delegate that is executed when an instance changes its selection state.
+	 *
+	 * @return The delegate.
+	 */
+	DECLARE_EVENT(ISessionManager, FInstanceSelectionChangedEvent)
+	virtual FInstanceSelectionChangedEvent& OnInstanceSelectionChanged() = 0;
+
+	/**
+	 * Returns a delegate that is executed when the selected session received a log message from one of its instances.
+	 *
+	 * @return The delegate.
+	 */
+	DECLARE_EVENT_ThreeParams(ISessionManager, FLogReceivedEvent, const TSharedRef<ISessionInfo>& /*OwnerSession*/, const TSharedRef<ISessionInstanceInfo>& /*SessionInstance*/, const TSharedRef<FSessionLogMessage>& /*LogMessage*/);
+	virtual FLogReceivedEvent& OnLogReceived() = 0;
+
+	/**
+	 * Returns a delegate that is executed when the selected session changed.
+	 *
+	 * @return The delegate.
+	 */
+	DECLARE_EVENT_OneParam(ISessionManager, FSelectedSessionChangedEvent, const TSharedPtr<ISessionInfo>& /*SelectedSession*/)
+	virtual FSelectedSessionChangedEvent& OnSelectedSessionChanged() = 0;
+
+	/**
+	 * Returns a delegate that is executed when the list of sessions has changed.
+	 *
+	 * @return The delegate.
+	 */
+	virtual FSimpleMulticastDelegate& OnSessionsUpdated() = 0;
+
+	/**
+	 * Returns a delegate that is executed when a session instance is updated.
+	 *
+	 * @return The delegate.
+	 */
+	virtual FSimpleMulticastDelegate& OnSessionInstanceUpdated() = 0;
+
+public:
+
+	/** Virtual destructor. */
+	virtual ~ISessionManager() { }
+};
+
+
+/** Type definition for shared pointers to instances of ISessionManager. */
+typedef TSharedPtr<ISessionManager> ISessionManagerPtr;
+
+/** Type definition for shared references to instances of ISessionManager. */
+typedef TSharedRef<ISessionManager> ISessionManagerRef;
diff --git a/Engine/Source/Developer/SessionServices/Public/ISessionService.h b/Engine/Source/Developer/SessionServices/Public/ISessionService.h
new file mode 100644
--- /dev/null
+++ b/Engine/Source/Developer/SessionServices/Public/ISessionService.h
@@ -0,0 +1,47 @@
+// Copyright 1998-2015 Epic Games, Inc. All Rights Reserved.
+
+#pragma once
+
+
+/**
+ * Interface for application session services.
+ */
+class ISessionService
+{
+public:
+
+	/**
+	 * Checks whether the service is running.
+	 *
+	 * @return true if the service is running, false otherwise.
+	 * @see Start, Stop
+	 */
+	virtual bool IsRunning() = 0;
+
+	/**
+	 * Starts the service.
+	 *
+	 * @return true if the service was started, false otherwise.
+	 * @see IsRunning, Stop
+	 */
+	virtual bool Start() = 0;
+
+	/**
+	 * Stops the service.
+	 *
+	 * @see IsRunning, Start
+	 */
+	virtual void Stop() = 0;
+
+public:
+
+	/** Virtual destructor. */
+	virtual ~ISessionService() { }
+};
+
+
+/** Type definition for shared pointers to instances of ISessionService. */
+typedef TSharedPtr<ISessionService> ISessionServicePtr;
+
+/** Type definition for shared references to instances of ISessionService. */
+typedef TSharedRef<ISessionService> ISessionServiceRef;
diff --git a/Engine/Source/Developer/SessionServices/Public/ISessionServicesModule.h b/Engine/Source/Developer/SessionServices/Public/ISessionServicesModule.h
new file mode 100644
--- /dev/null
+++ b/Engine/Source/Developer/SessionServices/Public/ISessionServicesModule.h
@@ -0,0 +1,37 @@
+// Copyright 1998-2015 Epic Games, Inc. All Rights Reserved.
+
+#pragma once
+
+
+// forward declarations
+class ISessionManager;
+class ISessionService;
+
+
+/**
+ * Interface for session core modules.
+ */
+class ISessionServicesModule
+	: public IModuleInterface
+{
+public:
+
+	/**
+	 * Gets the session manager.
+	 *
+	 * @return The session manager.
+	 */
+	virtual TSharedRef<ISessionManager> GetSessionManager() = 0;
+
+	/** 
+	 * Gets the session service.
+	 *
+	 * @return The session service.
+	 */
+	virtual TSharedRef<ISessionService> GetSessionService() = 0;
+
+public:
+
+	/** Virtual destructor. */
+	virtual ~ISessionServicesModule() { }
+};
diff --git a/Engine/Source/Developer/SessionServices/Public/SessionLogMessage.h b/Engine/Source/Developer/SessionServices/Public/SessionLogMessage.h
new file mode 100644
--- /dev/null
+++ b/Engine/Source/Developer/SessionServices/Public/SessionLogMessage.h
@@ -0,0 +1,73 @@
+// Copyright 1998-2015 Epic Games, Inc. All Rights Reserved.
+
+#pragma once
+
+
+/**
+ * Structure for log messages.
+ */
+struct FSessionLogMessage
+{
+	/** Holds the log category. */
+	FName Category;
+
+	/** Holds the identifier of the engine instance that generated this log message. */
+	FGuid InstanceId;
+
+	/** Holds the name of the engine instance that generated this log message. */
+	FString InstanceName;
+
+	/** Holds the message text. */
+	FString Text;
+
+	/** Holds the time at which the message was generated. */
+	FDateTime Time;
+
+	/** Holds the number of seconds from the start of the instance at which the message was generated. */
+	double TimeSeconds;
+
+	/** Holds the verbosity type. */
+	ELogVerbosity::Type Verbosity;
+
+public:
+
+	/**
+	 * Creates and initializes a new instance.
+	 *
+	 * @param InInstanceId The identifier of the instance that generated the log message.
+	 * @param InInstanceName The name of the engine instance that generated the log message.
+	 * @param InTimeSeconds The number of seconds from the start of the instance at which the message was generated.
+	 * @param InText The message text.
+	 * @param InVerbosity The verbosity type.
+	 * @param InCategory The log category.
+	 */
+	FSessionLogMessage(const FGuid& InInstanceId, const FString& InInstanceName, float InTimeSeconds, const FString& InText, ELogVerbosity::Type InVerbosity, const FName& InCategory)
+		: Category(InCategory)
+		, InstanceId(InInstanceId)
+		, InstanceName(InInstanceName)
+		, Text((InCategory != NAME_None) ? InCategory.ToString() + TEXT(": ") + InText : InText)
+		, Time(FDateTime::Now())
+		, TimeSeconds(InTimeSeconds)
+		, Verbosity(InVerbosity)
+	{ }
+
+public:
+
+	/**
+	 * Implements a predicate to compare two log messages by log time.
+	 */
+	struct TimeComparer
+	{
+		bool operator()(const TSharedPtr<FSessionLogMessage>& A, const TSharedPtr<FSessionLogMessage>& B) const
+		{
+			return A->Time < B->Time;
+		}
+	};
+};
+
+
+/** Type definition for shared pointers to instances of FLauncherLogMessage. */
+typedef TSharedPtr<FSessionLogMessage> FSessionLogMessagePtr;
+
+/** Type definition for shared references to instances of FLauncherLogMessage. */
+typedef TSharedRef<FSessionLogMessage> FSessionLogMessageRef;
