SessionServices: Removed obsolete delegate field

diff --git a/Engine/Source/Developer/SessionServices/Private/SessionManager.cpp b/Engine/Source/Developer/SessionServices/Private/SessionManager.cpp
--- a/Engine/Source/Developer/SessionServices/Private/SessionManager.cpp
+++ b/Engine/Source/Developer/SessionServices/Private/SessionManager.cpp
@@ -34,8 +34,7 @@ FSessionManager::FSessionManager( const IMessageBusRef& InMessageBus )
 		.Handling<FSessionServicePong>(this, &FSessionManager::HandleSessionPongMessage);
 
 	// initialize ticker
-	TickDelegate = FTickerDelegate::CreateRaw(this, &FSessionManager::HandleTicker);
-	TickDelegateHandle = FTicker::GetCoreTicker().AddTicker(TickDelegate, 1.f);
+	TickDelegateHandle = FTicker::GetCoreTicker().AddTicker(FTickerDelegate::CreateRaw(this, &FSessionManager::HandleTicker), 1.f);
 
 	SendPing();
 }
diff --git a/Engine/Source/Developer/SessionServices/Private/SessionManager.h b/Engine/Source/Developer/SessionServices/Private/SessionManager.h
--- a/Engine/Source/Developer/SessionServices/Private/SessionManager.h
+++ b/Engine/Source/Developer/SessionServices/Private/SessionManager.h
@@ -21,12 +21,12 @@ public:
 	 */
 	FSessionManager( const IMessageBusRef& InMessageBus );
 	
-	/** Default destructor */
+	/** Destructor */
 	~FSessionManager();
 
 public:
 
-	// Begin ISessionManager Interface
+	// ISessionManager interface
 
 	virtual void AddOwner( const FString& InOwner ) override;
 	virtual void GetSelectedInstances( TArray<TSharedPtr<ISessionInstanceInfo>>& OutInstances) const override;
@@ -124,7 +124,7 @@ private:
 	ISessionInfoPtr SelectedSession;
 
 	/** Holds the collection of discovered sessions. */
-	TMap<FGuid, TSharedPtr<FSessionInfo> > Sessions;
+	TMap<FGuid, TSharedPtr<FSessionInfo>> Sessions;
 
 private:
 
@@ -149,9 +149,6 @@ private:
 	/** Holds a delegate to be invoked when a session instance is updated. */
 	FSimpleMulticastDelegate SessionInstanceUpdatedDelegate;
 
-	/** Holds a delegate to be invoked when the widget ticks. */
-	FTickerDelegate TickDelegate;
-
 	/** Holds a delegate to be invoked when the widget ticks. */
 	FDelegateHandle TickDelegateHandle;
 };
