Added missing forward declaration

diff --git a/Engine/Source/Developer/SessionServices/Public/ISessionManager.h b/Engine/Source/Developer/SessionServices/Public/ISessionManager.h
--- a/Engine/Source/Developer/SessionServices/Public/ISessionManager.h
+++ b/Engine/Source/Developer/SessionServices/Public/ISessionManager.h
@@ -5,6 +5,7 @@
 
 class ISessionInfo;
 class ISessionInstanceInfo;
+struct FSessionLogMessage;
 
 
 /**
