SessionFrontend: Implemented session authorization for remote users

See SESSION console command for details. Some more updates are needed in Automation, Profiler, etc. to perform authorization on a per-instance instead of per-session basis.

diff --git a/Engine/Source/Developer/SessionMessages/Classes/SessionServiceMessages.h b/Engine/Source/Developer/SessionMessages/Classes/SessionServiceMessages.h
--- a/Engine/Source/Developer/SessionMessages/Classes/SessionServiceMessages.h
+++ b/Engine/Source/Developer/SessionMessages/Classes/SessionServiceMessages.h
@@ -15,6 +15,18 @@ USTRUCT()
 struct FSessionServicePing
 {
 	GENERATED_USTRUCT_BODY()
+
+	/** The name of the user who sent this ping. */
+	UPROPERTY(EditAnywhere, Category="Message")
+	FString UserName;
+
+	/** Default constructor. */
+	FSessionServicePing() { }
+
+	/** Creates and initializes a new instance. */
+	FSessionServicePing(const FString& InUserName)
+		: UserName(InUserName)
+	{ }
 };
 
 
@@ -26,6 +38,10 @@ struct FSessionServicePong
 {
 	GENERATED_USTRUCT_BODY()
 
+	/** Indicates whether the pinging user is authorized to interact with this session. */
+	UPROPERTY(EditAnywhere, Category="Message")
+	bool Authorized;
+
 	/** Holds the application's build date. */
 	UPROPERTY(EditAnywhere, Category="Message")
 	FString BuildDate;
@@ -42,7 +58,7 @@ struct FSessionServicePong
 	UPROPERTY(EditAnywhere, Category="Message")
 	FString InstanceName;
 
-	/** Holds a flag indicating whether the application is running on a console. */
+	/** Indicates whether the application is running on a console. */
 	UPROPERTY(EditAnywhere, Category="Message")
 	bool IsConsoleBuild;
 
@@ -62,7 +78,7 @@ struct FSessionServicePong
 	UPROPERTY(EditAnywhere, Category="Message")
 	FString SessionOwner;
 
-	/** Holds a flag indicating whether the application is the only one in that session. */
+	/** Indicates whether the application is the only one in that session. */
 	UPROPERTY(EditAnywhere, Category="Message")
 	bool Standalone;
 };
@@ -101,9 +117,7 @@ struct FSessionServiceLog
 
 public:
 
-	/**
-	 * Default constructor.
-	 */
+	/** Default constructor. */
 	FSessionServiceLog() { }
 
 	/**
diff --git a/Engine/Source/Developer/SessionServices/Private/SessionInstanceInfo.cpp b/Engine/Source/Developer/SessionServices/Private/SessionInstanceInfo.cpp
--- a/Engine/Source/Developer/SessionServices/Private/SessionInstanceInfo.cpp
+++ b/Engine/Source/Developer/SessionServices/Private/SessionInstanceInfo.cpp
@@ -6,8 +6,9 @@
 /* FSessionInstanceInfo structors
  *****************************************************************************/
 
-FSessionInstanceInfo::FSessionInstanceInfo( const FGuid& InInstanceId, const TSharedRef<ISessionInfo>& InOwner, const IMessageBusRef& InMessageBus )
-	: EngineVersion(0)
+FSessionInstanceInfo::FSessionInstanceInfo(const FGuid& InInstanceId, const TSharedRef<ISessionInfo>& InOwner, const IMessageBusRef& InMessageBus)
+	: Authorized(false)
+	, EngineVersion(0)
 	, InstanceId(InInstanceId)
 	, Owner(InOwner)
 {
@@ -19,7 +20,7 @@ FSessionInstanceInfo::FSessionInstanceInfo( const FGuid& InInstanceId, const TSh
 /* FSessionInstanceInfo interface
  *****************************************************************************/
 
-void FSessionInstanceInfo::UpdateFromMessage( const FEngineServicePong& Message, const IMessageContextRef& Context )
+void FSessionInstanceInfo::UpdateFromMessage(const FEngineServicePong& Message, const IMessageContextRef& Context)
 {
 	if (Message.InstanceId != InstanceId)
 	{
@@ -35,7 +36,7 @@ void FSessionInstanceInfo::UpdateFromMessage( const FEngineServicePong& Message,
 }
 
 
-void FSessionInstanceInfo::UpdateFromMessage( const FSessionServicePong& Message, const IMessageContextRef& Context )
+void FSessionInstanceInfo::UpdateFromMessage(const FSessionServicePong& Message, const IMessageContextRef& Context)
 {
 	if (Message.InstanceId != InstanceId)
 	{
@@ -47,6 +48,7 @@ void FSessionInstanceInfo::UpdateFromMessage( const FSessionServicePong& Message
 		MessageEndpoint->Send(new FSessionServiceLogSubscribe(), Context->GetSender());
 	}
 
+	Authorized = Message.Authorized;
 	ApplicationAddress = Context->GetSender();
 	BuildDate = Message.BuildDate;
 	DeviceName = Message.DeviceName;
@@ -61,7 +63,7 @@ void FSessionInstanceInfo::UpdateFromMessage( const FSessionServicePong& Message
 /* FSessionInstanceInfo interface
  *****************************************************************************/
 
-void FSessionInstanceInfo::ExecuteCommand( const FString& CommandString )
+void FSessionInstanceInfo::ExecuteCommand(const FString& CommandString)
 {
 	if (MessageEndpoint.IsValid() && EngineAddress.IsValid())
 	{
@@ -82,7 +84,7 @@ void FSessionInstanceInfo::Terminate()
 /* FSessionInstanceInfo callbacks
  *****************************************************************************/
 
-void FSessionInstanceInfo::HandleSessionLogMessage( const FSessionServiceLog& Message, const IMessageContextRef& Context )
+void FSessionInstanceInfo::HandleSessionLogMessage(const FSessionServiceLog& Message, const IMessageContextRef& Context)
 {
 	FSessionLogMessageRef LogMessage = MakeShareable(
 		new FSessionLogMessage(
diff --git a/Engine/Source/Developer/SessionServices/Private/SessionInstanceInfo.h b/Engine/Source/Developer/SessionServices/Private/SessionInstanceInfo.h
--- a/Engine/Source/Developer/SessionServices/Private/SessionInstanceInfo.h
+++ b/Engine/Source/Developer/SessionServices/Private/SessionInstanceInfo.h
@@ -28,7 +28,7 @@ public:
 	 * @param InOwner The session that owns this instance.
 	 * @param InMessageBus The message bus to use.
 	 */
-	FSessionInstanceInfo( const FGuid& InInstanceId, const TSharedRef<ISessionInfo>& InOwner, const IMessageBusRef& InMessageBus );
+	FSessionInstanceInfo(const FGuid& InInstanceId, const TSharedRef<ISessionInfo>& InOwner, const IMessageBusRef& InMessageBus);
 
 public:
 
@@ -38,7 +38,7 @@ public:
 	 * @param Message The message containing engine information.
 	 * @param Context The message context.
 	 */
-	void UpdateFromMessage( const FEngineServicePong& Message, const IMessageContextRef& Context );
+	void UpdateFromMessage(const FEngineServicePong& Message, const IMessageContextRef& Context);
 
 	/**
 	 * Updates this instance info with the data in the specified message.
@@ -46,13 +46,13 @@ public:
 	 * @param Message The message containing instance information.
 	 * @param Context The message context.
 	 */
-	void UpdateFromMessage( const FSessionServicePong& Message, const IMessageContextRef& Context );
+	void UpdateFromMessage(const FSessionServicePong& Message, const IMessageContextRef& Context);
 
 public:	
 
 	// IGameInstanceInfo interface
 
-	virtual void ExecuteCommand( const FString& CommandString ) override;
+	virtual void ExecuteCommand(const FString& CommandString) override;
 
 	virtual const FString& GetBuildDate() const override
 	{
@@ -114,6 +114,11 @@ public:
 		return WorldTimeSeconds;
 	}
 
+	virtual bool IsAuthorized() const override
+	{
+		return Authorized;
+	}
+
 	virtual const bool IsConsole() const override
 	{
 		return IsConsoleBuild;
@@ -135,13 +140,16 @@ public:
 private:
 
 	/** Handles FSessionServiceLog messages. */
-	void HandleSessionLogMessage( const FSessionServiceLog& Message, const IMessageContextRef& Context );
+	void HandleSessionLogMessage(const FSessionServiceLog& Message, const IMessageContextRef& Context);
 
 private:
 
 	/** Holds the message bus address of the application instance. */
 	FMessageAddress ApplicationAddress;
 
+	/** Whether the current user is authorized to interact with this instance. */
+	bool Authorized;
+
 	/** Holds the instance's build date. */
 	FString BuildDate;
 
diff --git a/Engine/Source/Developer/SessionServices/Private/SessionManager.cpp b/Engine/Source/Developer/SessionServices/Private/SessionManager.cpp
--- a/Engine/Source/Developer/SessionServices/Private/SessionManager.cpp
+++ b/Engine/Source/Developer/SessionServices/Private/SessionManager.cpp
@@ -229,7 +229,7 @@ void FSessionManager::SendPing()
 	if (MessageEndpoint.IsValid())
 	{
 		MessageEndpoint->Publish(new FEngineServicePing(), EMessageScope::Network);
-		MessageEndpoint->Publish(new FSessionServicePing(), EMessageScope::Network);
+		MessageEndpoint->Publish(new FSessionServicePing(FPlatformProcess::UserName(false)), EMessageScope::Network);
 	}
 
 	LastPingTime = FDateTime::UtcNow();
diff --git a/Engine/Source/Developer/SessionServices/Private/SessionService.cpp b/Engine/Source/Developer/SessionServices/Private/SessionService.cpp
--- a/Engine/Source/Developer/SessionServices/Private/SessionService.cpp
+++ b/Engine/Source/Developer/SessionServices/Private/SessionService.cpp
@@ -6,7 +6,7 @@
 /* FSessionService structors
  *****************************************************************************/
 
-FSessionService::FSessionService( const IMessageBusRef& InMessageBus )
+FSessionService::FSessionService(const IMessageBusRef& InMessageBus)
 	: MessageBusPtr(InMessageBus)
 { }
 
@@ -60,53 +60,61 @@ void FSessionService::Stop()
 /* FSessionService implementation
  *****************************************************************************/
 
-void FSessionService::SendLog( const TCHAR* Data, ELogVerbosity::Type Verbosity, const class FName& Category )
+void FSessionService::SendLog(const TCHAR* Data, ELogVerbosity::Type Verbosity, const class FName& Category)
 {
-	if (MessageEndpoint.IsValid())
+	if (!MessageEndpoint.IsValid())
 	{
-		FScopeLock Lock(&LogSubscribersLock);
-
-		if (LogSubscribers.Num() > 0)
-		{
-			MessageEndpoint->Send(
-				new FSessionServiceLog(
-					Category,
-					Data,
-					FApp::GetInstanceId(),
-					FPlatformTime::Seconds() - GStartTime,
-					Verbosity
-				),
-				LogSubscribers
-			);
-		}
+		return;
 	}
-}
 
+	FScopeLock Lock(&LogSubscribersLock);
 
-void FSessionService::SendNotification( const TCHAR* NotificationText, const FMessageAddress& Recipient )
-{
-	if (MessageEndpoint.IsValid())
+	if (LogSubscribers.Num() > 0)
 	{
 		MessageEndpoint->Send(
 			new FSessionServiceLog(
-				FName("RemoteSession"),
-				NotificationText,
+				Category,
+				Data,
 				FApp::GetInstanceId(),
 				FPlatformTime::Seconds() - GStartTime,
-				ELogVerbosity::Display
+				Verbosity
 			),
-			Recipient
+			LogSubscribers
 		);
 	}
 }
 
 
-void FSessionService::SendPong( const IMessageContextRef& Context )
+void FSessionService::SendNotification(const TCHAR* NotificationText, const FMessageAddress& Recipient)
 {
-	if (MessageEndpoint.IsValid())
+	if (!MessageEndpoint.IsValid())
 	{
-		FSessionServicePong* Message = new FSessionServicePong();
+		return;
+	}
+
+	MessageEndpoint->Send(
+		new FSessionServiceLog(
+			FName("RemoteSession"),
+			NotificationText,
+			FApp::GetInstanceId(),
+			FPlatformTime::Seconds() - GStartTime,
+			ELogVerbosity::Display
+		),
+		Recipient
+	);
+}
+
 
+void FSessionService::SendPong(const IMessageContextRef& Context, const FString& UserName)
+{
+	if (!MessageEndpoint.IsValid())
+	{
+		return;
+	}
+
+	FSessionServicePong* Message = new FSessionServicePong();
+	{
+		Message->Authorized = FApp::IsAuthorizedUser(UserName);
 		Message->BuildDate = FApp::GetBuildDate();
 		Message->DeviceName = FPlatformProcess::ComputerName();
 		Message->InstanceId = FApp::GetInstanceId();
@@ -123,9 +131,9 @@ void FSessionService::SendPong( const IMessageContextRef& Context )
 		Message->SessionName = FApp::GetSessionName();
 		Message->SessionOwner = FApp::GetSessionOwner();
 		Message->Standalone = FApp::IsStandalone();
-
-		MessageEndpoint->Send(Message, Context->GetSender());
 	}
+
+	MessageEndpoint->Send(Message, Context->GetSender());
 }
 
 
@@ -138,21 +146,21 @@ void FSessionService::HandleMessageEndpointShutdown()
 }
 
 
-void FSessionService::HandleSessionLogSubscribeMessage( const FSessionServiceLogSubscribe& Message, const IMessageContextRef& Context )
+void FSessionService::HandleSessionLogSubscribeMessage(const FSessionServiceLogSubscribe& Message, const IMessageContextRef& Context)
 {
 	FScopeLock Lock(&LogSubscribersLock);
 	LogSubscribers.AddUnique(Context->GetSender());
 }
 
 
-void FSessionService::HandleSessionLogUnsubscribeMessage( const FSessionServiceLogUnsubscribe& Message, const IMessageContextRef& Context )
+void FSessionService::HandleSessionLogUnsubscribeMessage(const FSessionServiceLogUnsubscribe& Message, const IMessageContextRef& Context)
 {
 	FScopeLock Lock(&LogSubscribersLock);
 	LogSubscribers.Remove(Context->GetSender());
 }
 
 
-void FSessionService::HandleSessionPingMessage( const FSessionServicePing& Message, const IMessageContextRef& Context )
+void FSessionService::HandleSessionPingMessage(const FSessionServicePing& Message, const IMessageContextRef& Context)
 {
-	SendPong(Context);
+	SendPong(Context, Message.UserName);
 }
diff --git a/Engine/Source/Developer/SessionServices/Private/SessionService.h b/Engine/Source/Developer/SessionServices/Private/SessionService.h
--- a/Engine/Source/Developer/SessionServices/Private/SessionService.h
+++ b/Engine/Source/Developer/SessionServices/Private/SessionService.h
@@ -19,16 +19,16 @@ public:
 	 *
 	 * @param InMessageBus The message bus to use.
 	 */
-	FSessionService( const IMessageBusRef& InMessageBus );
+	FSessionService(const IMessageBusRef& InMessageBus);
 
-	/** Destructor. */
-	~FSessionService();
+	/** Virtual destructor. */
+	virtual ~FSessionService();
 
 public:
 
 	// FOutputDevice interface
 
-	virtual void Serialize( const TCHAR* Data, ELogVerbosity::Type Verbosity, const class FName& Category ) override
+	virtual void Serialize(const TCHAR* Data, ELogVerbosity::Type Verbosity, const class FName& Category) override
 	{
 		SendLog(Data, Verbosity, Category);
 	}
@@ -54,7 +54,7 @@ protected:
 	 * @param Verbosity The verbosity type.
 	 * @param Category The log category.
 	 */
-	void SendLog( const TCHAR* Data, ELogVerbosity::Type Verbosity = ELogVerbosity::Log, const class FName& Category = "Log" );
+	void SendLog(const TCHAR* Data, ELogVerbosity::Type Verbosity = ELogVerbosity::Log, const class FName& Category = "Log");
 
 	/**
 	 * Sends a notification to the specified recipient.
@@ -62,14 +62,15 @@ protected:
 	 * @param NotificationText The notification text.
 	 * @param Recipient The recipient's message address.
 	 */
-	void SendNotification( const TCHAR* NotificationText, const FMessageAddress& Recipient );
+	void SendNotification(const TCHAR* NotificationText, const FMessageAddress& Recipient);
 
 	/**
 	 * Publishes a ping response.
 	 *
 	 * @param Context The context of the received Ping message.
+	 * @param UserName The name of the user that sent the ping.
 	 */
-	void SendPong( const IMessageContextRef& Context );
+	void SendPong(const IMessageContextRef& Context, const FString& UserName);
 
 private:
 
@@ -77,13 +78,13 @@ private:
 	void HandleMessageEndpointShutdown();
 
 	/** Handles FSessionServiceLogSubscribe messages. */
-	void HandleSessionLogSubscribeMessage( const FSessionServiceLogSubscribe& Message, const IMessageContextRef& Context );
+	void HandleSessionLogSubscribeMessage(const FSessionServiceLogSubscribe& Message, const IMessageContextRef& Context);
 
 	/** Handles FSessionServiceLogUnsubscribe messages. */
-	void HandleSessionLogUnsubscribeMessage( const FSessionServiceLogUnsubscribe& Message, const IMessageContextRef& Context );
+	void HandleSessionLogUnsubscribeMessage(const FSessionServiceLogUnsubscribe& Message, const IMessageContextRef& Context);
 
 	/** Handles FSessionServicePing messages. */
-	void HandleSessionPingMessage( const FSessionServicePing& Message, const IMessageContextRef& Context );
+	void HandleSessionPingMessage(const FSessionServicePing& Message, const IMessageContextRef& Context);
 
 private:
 
diff --git a/Engine/Source/Developer/SessionServices/Private/SessionServicesModule.cpp b/Engine/Source/Developer/SessionServices/Private/SessionServicesModule.cpp
--- a/Engine/Source/Developer/SessionServices/Private/SessionServicesModule.cpp
+++ b/Engine/Source/Developer/SessionServices/Private/SessionServicesModule.cpp
@@ -28,40 +28,55 @@ public:
 
 	virtual bool Exec(UWorld* InWorld, const TCHAR* Cmd, FOutputDevice& Ar) override
 	{
-		if (FParse::Command(&Cmd, TEXT("SESSION")))
+		if (!FParse::Command(&Cmd, TEXT("SESSION")))
 		{
-			if (FParse::Command(&Cmd, TEXT("STATUS")))
-			{
-				Ar.Logf(TEXT("Instance ID: %s"), *FApp::GetInstanceId().ToString());
-				Ar.Logf(TEXT("Instance Name: %s"), *FApp::GetInstanceName());
-				Ar.Logf(TEXT("Session ID: %s"), *FApp::GetSessionId().ToString());
-				Ar.Logf(TEXT("Session Name: %s"), *FApp::GetSessionName());
-				Ar.Logf(TEXT("Session Owner: %s"), *FApp::GetSessionOwner());
-				Ar.Logf(TEXT("Standalone: %s"), FApp::IsStandalone() ? *GYes.ToString() : *GNo.ToString());
-			}
-			else if (FParse::Command(&Cmd, TEXT("SETNAME")))
-			{
-				FApp::SetSessionName(FParse::Token(Cmd, false));
-			}
-			else if (FParse::Command(&Cmd, TEXT("SETOWNER")))
-			{
-				FApp::SetSessionOwner(FParse::Token(Cmd, false));
-			}
-			else
-			{
-				// show usage
-				Ar.Log(TEXT("Usage: SESSION <Command>"));
-				Ar.Log(TEXT(""));
-				Ar.Log(TEXT("Command"));
-				Ar.Log(TEXT("    STATUS = Displays the status of this application session"));
-				Ar.Log(TEXT("    SETNAME <Name> = Sets the name of this instance"));
-				Ar.Log(TEXT("    SETOWNER <Owner> = Sets the name of the owner of this instance"));
-			}
-
-			return true;
+			return false;
 		}
 
-		return false;
+		if (FParse::Command(&Cmd, TEXT("AUTH")))
+		{
+			FApp::AuthorizeUser(FParse::Token(Cmd, false));
+		}
+		else if (FParse::Command(&Cmd, TEXT("DENY")))
+		{
+			FApp::DenyUser(FParse::Token(Cmd, false));
+		}
+		else if (FParse::Command(&Cmd, TEXT("DENYALL")))
+		{
+			FApp::DenyAllUsers();
+		}
+		else if (FParse::Command(&Cmd, TEXT("STATUS")))
+		{
+			Ar.Logf(TEXT("Instance ID: %s"), *FApp::GetInstanceId().ToString());
+			Ar.Logf(TEXT("Instance Name: %s"), *FApp::GetInstanceName());
+			Ar.Logf(TEXT("Session ID: %s"), *FApp::GetSessionId().ToString());
+			Ar.Logf(TEXT("Session Name: %s"), *FApp::GetSessionName());
+			Ar.Logf(TEXT("Session Owner: %s"), *FApp::GetSessionOwner());
+			Ar.Logf(TEXT("Standalone: %s"), FApp::IsStandalone() ? *GYes.ToString() : *GNo.ToString());
+		}
+		else if (FParse::Command(&Cmd, TEXT("SETNAME")))
+		{
+			FApp::SetSessionName(FParse::Token(Cmd, false));
+		}
+		else if (FParse::Command(&Cmd, TEXT("SETOWNER")))
+		{
+			FApp::SetSessionOwner(FParse::Token(Cmd, false));
+		}
+		else
+		{
+			// show usage
+			Ar.Log(TEXT("Usage: SESSION <Command>"));
+			Ar.Log(TEXT(""));
+			Ar.Log(TEXT("Command"));
+			Ar.Log(TEXT("    AUTH <UserName> = Authorize a user to interact with this instance"));
+			Ar.Log(TEXT("    DENY <UserName> = Unauthorize a user from interacting with this instance"));
+			Ar.Log(TEXT("    DENYALL = Removes all authorized session users"));
+			Ar.Log(TEXT("    SETNAME <Name> = Sets the name of this instance"));
+			Ar.Log(TEXT("    SETOWNER <Owner> = Sets the name of the owner of this instance"));
+			Ar.Log(TEXT("    STATUS = Displays the status of this application session"));
+		}
+
+		return true;
 	}
 
 public:
diff --git a/Engine/Source/Developer/SessionServices/Public/ISessionInstanceInfo.h b/Engine/Source/Developer/SessionServices/Public/ISessionInstanceInfo.h
--- a/Engine/Source/Developer/SessionServices/Public/ISessionInstanceInfo.h
+++ b/Engine/Source/Developer/SessionServices/Public/ISessionInstanceInfo.h
@@ -106,7 +106,14 @@ public:
 	virtual float GetWorldTimeSeconds() const = 0;
 
 	/**
-	 * Checks whether this instance is a console build (i.e. no editor features).
+	 * Check whether the current user is authorized to interact with this instance.
+	 *
+	 * @return true if the uesr is authorized, false otherwise.
+	 */
+	virtual bool IsAuthorized() const = 0;
+
+	/**
+	 * Check whether this instance is a console build (i.e. no editor features).
 	 *
 	 * @return true if it is a console build, false otherwise.
 	 */
diff --git a/Engine/Source/Runtime/Core/Private/Misc/App.cpp b/Engine/Source/Runtime/Core/Private/Misc/App.cpp
--- a/Engine/Source/Runtime/Core/Private/Misc/App.cpp
+++ b/Engine/Source/Runtime/Core/Private/Misc/App.cpp
@@ -12,6 +12,7 @@ FGuid FApp::InstanceId = FGuid::NewGuid();
 FGuid FApp::SessionId = FGuid::NewGuid();
 FString FApp::SessionName = FString();
 FString FApp::SessionOwner = FString();
+TArray<FString> FApp::SessionUsers = TArray<FString>();
 bool FApp::Standalone = true;
 bool FApp::bIsBenchmarking = false;
 bool FApp::bUseFixedTimeStep = false;
diff --git a/Engine/Source/Runtime/Core/Public/Misc/App.h b/Engine/Source/Runtime/Core/Public/Misc/App.h
--- a/Engine/Source/Runtime/Core/Public/Misc/App.h
+++ b/Engine/Source/Runtime/Core/Public/Misc/App.h
@@ -133,6 +133,38 @@ public:
 
 public:
 
+	/**
+	 * Add the specified user to the list of authorized session users.
+	 *
+	 * @param UserName The name of the user to add.
+	 * @see DenyUser, DenyAllUsers, IsAuthorizedUser
+	 */
+	FORCEINLINE static void AuthorizeUser(const FString& UserName)
+	{
+		SessionUsers.AddUnique(UserName);
+	}
+
+	/**
+	 * Removes all authorized users.
+	 *
+	 * @see AuthorizeUser, DenyUser, IsAuthorizedUser
+	 */
+	FORCEINLINE static void DenyAllUsers()
+	{
+		SessionUsers.Empty();
+	}
+
+	/**
+	 * Remove the specified user from the list of authorized session users.
+	 *
+	 * @param UserName The name of the user to remove.
+	 * @see AuthorizeUser, DenyAllUsers, IsAuthorizedUser
+	 */
+	FORCEINLINE static void DenyUser(const FString& UserName)
+	{
+		SessionUsers.Remove(UserName);
+	}
+
 	/**
 	 * Gets the globally unique identifier of this application instance.
 	 *
@@ -205,6 +237,18 @@ public:
 	 */
 	static void InitializeSession();
 
+	/**
+	 * Check whether the specified user is authorized to interact with this session.
+	 *
+	 * @param UserName The name of the user to check.
+	 * @return true if the user is authorized, false otherwise.
+	 * @see AuthorizeUser, DenyUser, DenyAllUsers
+	 */
+	FORCEINLINE static bool IsAuthorizedUser(const FString& UserName)
+	{
+		return ((FPlatformProcess::UserName(false) == UserName) || SessionUsers.Contains(UserName));
+	}
+
 	/**
 	 * Checks whether this is a standalone application.
 	 *
@@ -485,6 +529,9 @@ private:
 	/** Holds the name of the user that launched session. */
 	static FString SessionOwner;
 
+	/** List of authorized session users. */
+	static TArray<FString> SessionUsers;
+
 	/** Holds a flag indicating whether this is a standalone session. */
 	static bool Standalone;
 
